Mono Video Live Link device implementation

diff --git a/Engine/Plugins/VirtualProduction/CaptureManager/CaptureManagerApp/Source/CaptureManagerEditor/Private/Scripting/ULiveLinkHubCaptureDevice.cpp b/Engine/Plugins/VirtualProduction/CaptureManager/CaptureManagerApp/Source/CaptureManagerEditor/Private/Scripting/ULiveLinkHubCaptureDevice.cpp
--- a/Engine/Plugins/VirtualProduction/CaptureManager/CaptureManagerApp/Source/CaptureManagerEditor/Private/Scripting/ULiveLinkHubCaptureDevice.cpp
+++ b/Engine/Plugins/VirtualProduction/CaptureManager/CaptureManagerApp/Source/CaptureManagerEditor/Private/Scripting/ULiveLinkHubCaptureDevice.cpp
@@ -315,6 +315,28 @@ ULiveLinkHubCaptureDevice* ULiveLinkHubCaptureDeviceFactory::CreateNetworkDevice
 	return Device;
 }
 
+ULiveLinkHubCaptureDevice* ULiveLinkHubCaptureDeviceFactory::CreateDeviceByClass(FString InName, UClass* InDeviceClass, ULiveLinkDeviceSettings* InSettings)
+{
+	TSubclassOf<ULiveLinkDevice> DeviceClass = InDeviceClass;
+
+	ULiveLinkDeviceSubsystem* DeviceSubsystem = GEngine->GetEngineSubsystem<ULiveLinkDeviceSubsystem>();
+	check(DeviceSubsystem);
+
+	TObjectPtr<ULiveLinkDevice> NewDevice = NewObject<ULiveLinkDevice>(DeviceSubsystem, DeviceClass);
+
+	check(NewDevice->Implements<ULiveLinkDeviceCapability_Ingest>());
+
+	const FGuid DeviceGuid = DeviceSubsystem->AddDevice(NewDevice);
+
+	NewDevice->OnDeviceAdded(DeviceGuid, InSettings);
+
+	TObjectPtr<ULiveLinkHubCaptureDevice> Device = NewObject<ULiveLinkHubCaptureDevice>(this);
+	Device->Name = MoveTemp(InName);
+	Device->ImplPtr->BackgroundDevice = MoveTemp(NewDevice);
+
+	return Device;
+}
+
 ULiveLinkHubCaptureDevice::ULiveLinkHubCaptureDevice() :
 	ImplPtr(MakePimpl<FImpl>())
 {
diff --git a/Engine/Plugins/VirtualProduction/CaptureManager/CaptureManagerApp/Source/CaptureManagerEditor/Private/Scripting/ULiveLinkHubCaptureDevice.h b/Engine/Plugins/VirtualProduction/CaptureManager/CaptureManagerApp/Source/CaptureManagerEditor/Private/Scripting/ULiveLinkHubCaptureDevice.h
--- a/Engine/Plugins/VirtualProduction/CaptureManager/CaptureManagerApp/Source/CaptureManagerEditor/Private/Scripting/ULiveLinkHubCaptureDevice.h
+++ b/Engine/Plugins/VirtualProduction/CaptureManager/CaptureManagerApp/Source/CaptureManagerEditor/Private/Scripting/ULiveLinkHubCaptureDevice.h
@@ -52,6 +52,9 @@ public:
 
 	UFUNCTION(BlueprintCallable, Category = "Ingest | Creator")
 	ULiveLinkHubCaptureDevice* CreateNetworkDevice(FString InName, FString InIPAddress, int32 InPort);
+
+	UFUNCTION(BlueprintCallable, Category = "Ingest | Creator")
+	ULiveLinkHubCaptureDevice* CreateDeviceByClass(FString InName, UClass* InDeviceClass, ULiveLinkDeviceSettings* Settings);
 };
 
 UCLASS()
diff --git a/Engine/Plugins/VirtualProduction/CaptureManager/CaptureManagerApp/Source/CaptureManagerMediaRW/CaptureManagerMediaRW.Build.cs b/Engine/Plugins/VirtualProduction/CaptureManager/CaptureManagerApp/Source/CaptureManagerMediaRW/CaptureManagerMediaRW.Build.cs
--- a/Engine/Plugins/VirtualProduction/CaptureManager/CaptureManagerApp/Source/CaptureManagerMediaRW/CaptureManagerMediaRW.Build.cs
+++ b/Engine/Plugins/VirtualProduction/CaptureManager/CaptureManagerApp/Source/CaptureManagerMediaRW/CaptureManagerMediaRW.Build.cs
@@ -22,7 +22,8 @@ public class CaptureManagerMediaRW : ModuleRules
 			"ImageCore",
 			"OodleDataCompression",
 			"Json",
-			"JsonUtilities"
+			"JsonUtilities",
+			"DataIngestCore"
 		});
 
 		if (Target.Platform == UnrealTargetPlatform.Win64)
diff --git a/Engine/Plugins/VirtualProduction/CaptureManager/CaptureManagerApp/Source/CaptureManagerMediaRW/Private/Readers/MHADepthVideoReader.cpp b/Engine/Plugins/VirtualProduction/CaptureManager/CaptureManagerApp/Source/CaptureManagerMediaRW/Private/Readers/MHADepthVideoReader.cpp
--- a/Engine/Plugins/VirtualProduction/CaptureManager/CaptureManagerApp/Source/CaptureManagerMediaRW/Private/Readers/MHADepthVideoReader.cpp
+++ b/Engine/Plugins/VirtualProduction/CaptureManager/CaptureManagerApp/Source/CaptureManagerMediaRW/Private/Readers/MHADepthVideoReader.cpp
@@ -10,6 +10,8 @@
 #include "HAL/PlatformFileManager.h"
 #include "Misc/FileHelper.h"
 
+#include "Utils/ParseTakeUtils.h"
+
 #include <oodle2.h>
 
 #define LOCTEXT_NAMESPACE "MHADepthVideoReader"
@@ -154,6 +156,14 @@ TOptional<FText> FMHADepthVideoReader::Open(const FString& InFileName)
 			DepthDimensions->ToSharedRef()->TryGetNumberField(TEXT("Width"), Dimensions.X);
 			DepthDimensions->ToSharedRef()->TryGetNumberField(TEXT("Height"), Dimensions.Y);
 		}
+
+		float FloatFrameRate = 0.f;
+		MetadataJson->TryGetNumberField(TEXT("DepthFrameRate"), FloatFrameRate);
+
+		if (!FMath::IsNearlyZero(FloatFrameRate))
+		{
+			FrameRate = UE::CaptureManager::ParseFrameRate(FloatFrameRate);
+		}
 	}
 
 	// Return the seek pointer to the start
@@ -253,4 +263,9 @@ FIntPoint FMHADepthVideoReader::GetDimensions() const
 	return Dimensions;
 }
 
+FFrameRate FMHADepthVideoReader::GetFrameRate() const
+{
+	return FrameRate;
+}
+
 #undef LOCTEXT_NAMESPACE
\ No newline at end of file
diff --git a/Engine/Plugins/VirtualProduction/CaptureManager/CaptureManagerApp/Source/CaptureManagerMediaRW/Private/Readers/MHADepthVideoReader.h b/Engine/Plugins/VirtualProduction/CaptureManager/CaptureManagerApp/Source/CaptureManagerMediaRW/Private/Readers/MHADepthVideoReader.h
--- a/Engine/Plugins/VirtualProduction/CaptureManager/CaptureManagerApp/Source/CaptureManagerMediaRW/Private/Readers/MHADepthVideoReader.h
+++ b/Engine/Plugins/VirtualProduction/CaptureManager/CaptureManagerApp/Source/CaptureManagerMediaRW/Private/Readers/MHADepthVideoReader.h
@@ -35,10 +35,12 @@ public:
 
 	virtual FTimespan GetDuration() const override;
 	virtual FIntPoint GetDimensions() const override;
+	virtual FFrameRate GetFrameRate() const override;
 
 private:
 
 	TUniquePtr<IFileHandle> ReadHandle;
 	FIntPoint Dimensions;
+	FFrameRate FrameRate = FFrameRate(30'000, 1'000);
 	EMediaOrientation Orientation;
 };
\ No newline at end of file
diff --git a/Engine/Plugins/VirtualProduction/CaptureManager/CaptureManagerApp/Source/CaptureManagerMediaRW/Private/Readers/WindowsMediaReader.cpp b/Engine/Plugins/VirtualProduction/CaptureManager/CaptureManagerApp/Source/CaptureManagerMediaRW/Private/Readers/WindowsMediaReader.cpp
--- a/Engine/Plugins/VirtualProduction/CaptureManager/CaptureManagerApp/Source/CaptureManagerMediaRW/Private/Readers/WindowsMediaReader.cpp
+++ b/Engine/Plugins/VirtualProduction/CaptureManager/CaptureManagerApp/Source/CaptureManagerMediaRW/Private/Readers/WindowsMediaReader.cpp
@@ -251,6 +251,12 @@ TOptional<FText> FWindowsVideoReader::Open(const FString& InFileName)
 	Result = ::MFGetAttributeSize(VideoMediaType, MF_MT_FRAME_SIZE, &VideoWidth, &VideoHeight);
 	WINAR_CHECK_AND_RETURN_ERROR_MESSAGE(Result, LOCTEXT("OpenVideo_FailedToObtainDimensions", "Failed to obtain the video resolution"));
 
+	UINT32 FrameRateNumerator = 0;
+	UINT32 FrameRateDenominator = 0;
+
+	Result = ::MFGetAttributeRatio(VideoMediaType, MF_MT_FRAME_RATE, &FrameRateNumerator, &FrameRateDenominator);
+	WINAR_CHECK_AND_RETURN_ERROR_MESSAGE(Result, LOCTEXT("OpenVideo_FailedToObtainFrameRate", "Failed to obtain the video frame rate"));
+
 	Result = VideoMediaType->GetGUID(MF_MT_SUBTYPE, &InputVideoSubType);
 	WINAR_CHECK_AND_RETURN_ERROR_MESSAGE(Result, LOCTEXT("OpenVideo_FailedToObtainMediaSubtype", "Failed to obtain the video subtype"));
 
@@ -287,6 +293,7 @@ TOptional<FText> FWindowsVideoReader::Open(const FString& InFileName)
 	WINAR_CHECK_AND_RETURN_ERROR_MESSAGE(Result, LOCTEXT("OpenVideo_FailedToSetMediaType", "Failed to set the media type to the reader"));
 
 	Dimensions = FIntPoint(VideoWidth, VideoHeight);
+	FrameRate = FFrameRate(FrameRateNumerator, FrameRateDenominator);
 
 	return {};
 }
@@ -403,6 +410,11 @@ FIntPoint FWindowsVideoReader::GetDimensions() const
 	return Dimensions;
 }
 
+FFrameRate FWindowsVideoReader::GetFrameRate() const
+{
+	return FrameRate;
+}
+
 #undef LOCTEXT_NAMESPACE
 
 #endif // PLATFORM_WINDOWS && !UE_SERVER
\ No newline at end of file
diff --git a/Engine/Plugins/VirtualProduction/CaptureManager/CaptureManagerApp/Source/CaptureManagerMediaRW/Private/Readers/WindowsMediaReader.h b/Engine/Plugins/VirtualProduction/CaptureManager/CaptureManagerApp/Source/CaptureManagerMediaRW/Private/Readers/WindowsMediaReader.h
--- a/Engine/Plugins/VirtualProduction/CaptureManager/CaptureManagerApp/Source/CaptureManagerMediaRW/Private/Readers/WindowsMediaReader.h
+++ b/Engine/Plugins/VirtualProduction/CaptureManager/CaptureManagerApp/Source/CaptureManagerMediaRW/Private/Readers/WindowsMediaReader.h
@@ -77,6 +77,7 @@ public:
 
 	virtual FTimespan GetDuration() const override;
 	virtual FIntPoint GetDimensions() const override;
+	virtual FFrameRate GetFrameRate() const override;
 
 private:
 
@@ -84,6 +85,7 @@ private:
 
 	FTimespan Duration;
 	FIntPoint Dimensions;
+	FFrameRate FrameRate;
 	UE::CaptureManager::EMediaTexturePixelFormat PixelFormat;
 	GUID InputVideoSubType;
 };
diff --git a/Engine/Plugins/VirtualProduction/CaptureManager/CaptureManagerApp/Source/CaptureManagerMediaRW/Public/IMediaReader.h b/Engine/Plugins/VirtualProduction/CaptureManager/CaptureManagerApp/Source/CaptureManagerMediaRW/Public/IMediaReader.h
--- a/Engine/Plugins/VirtualProduction/CaptureManager/CaptureManagerApp/Source/CaptureManagerMediaRW/Public/IMediaReader.h
+++ b/Engine/Plugins/VirtualProduction/CaptureManager/CaptureManagerApp/Source/CaptureManagerMediaRW/Public/IMediaReader.h
@@ -34,6 +34,7 @@ public:
 
 	virtual FTimespan GetDuration() const = 0;
 	virtual FIntPoint GetDimensions() const = 0;
+	virtual FFrameRate GetFrameRate() const = 0;
 };
 
 class CAPTUREMANAGERMEDIARW_API ICalibrationReader
diff --git a/Engine/Plugins/VirtualProduction/CaptureManager/CaptureManagerDevices/CaptureManagerDevices.uplugin b/Engine/Plugins/VirtualProduction/CaptureManager/CaptureManagerDevices/CaptureManagerDevices.uplugin
new file mode 100644
--- /dev/null
+++ b/Engine/Plugins/VirtualProduction/CaptureManager/CaptureManagerDevices/CaptureManagerDevices.uplugin
@@ -0,0 +1,32 @@
+{
+	"FileVersion": 1,
+	"Version": 1,
+	"VersionName": "1.0.0",
+	"FriendlyName": "Capture Manager Devices",
+	"Description": "The Capture Manager Devices contains devices that can be used from the Capture Manager layout of the LiveLink Hub",
+	"Category": "Virtual Production",
+	"CreatedBy": "Epic Games, Inc.",
+	"CreatedByURL": "http://epicgames.com",
+	"DocsURL": "",
+	"MarketplaceURL": "",
+	"SupportURL": "",
+	"CanContainContent": true,
+	"EnabledByDefault": false,
+	"IsBetaVersion": false,
+	"IsExperimentalVersion": true,
+	"Installed": false,
+	"Hidden" :  true,
+	"Modules": [
+		{
+			"Name": "MonoVideoLiveLinkDevice",
+			"Type": "Editor",
+			"LoadingPhase": "Default"
+		}
+	],
+	"Plugins": [
+		{
+			"Name": "CaptureManagerApp",
+			"Enabled": true
+		}
+	]
+}
\ No newline at end of file
diff --git a/Engine/Plugins/VirtualProduction/CaptureManager/CaptureManagerDevices/Source/MonoVideoLiveLinkDevice/MonoVideoLiveLinkDevice.Build.cs b/Engine/Plugins/VirtualProduction/CaptureManager/CaptureManagerDevices/Source/MonoVideoLiveLinkDevice/MonoVideoLiveLinkDevice.Build.cs
new file mode 100644
--- /dev/null
+++ b/Engine/Plugins/VirtualProduction/CaptureManager/CaptureManagerDevices/Source/MonoVideoLiveLinkDevice/MonoVideoLiveLinkDevice.Build.cs
@@ -0,0 +1,42 @@
+// Copyright Epic Games, Inc. All Rights Reserved.
+
+using UnrealBuildTool;
+using System.IO;
+
+public class MonoVideoLiveLinkDevice : ModuleRules
+{
+	public MonoVideoLiveLinkDevice(ReadOnlyTargetRules Target) : base(Target)
+	{
+		PCHUsage = PCHUsageMode.UseExplicitOrSharedPCHs;
+		ShortName = "MonoVideoLLD";
+
+		PublicDependencyModuleNames.AddRange(new string[]
+		{
+			"CaptureUtils",
+			"Core"
+		});
+
+		PrivateDependencyModuleNames.AddRange(new string[]
+		{
+			"Engine",
+			"Json",
+			"CaptureManagerMediaRW",
+			"IngestLiveLinkDevice",
+			"LiveLinkDevice",
+			"LiveLinkCapabilities",
+			"Media",
+			"LiveLinkHubCaptureMessaging",
+			"LiveLinkHubExportClient",
+			"LiveLinkHub",
+			"ToolWidgets",
+			"Slate",
+			"SlateCore",
+			"CaptureManagerTakeMetadata",
+			"CaptureManagerSettings",
+			"CoreUObject",
+			"UnrealEd"
+		});
+
+		CppCompileWarningSettings.UnsafeTypeCastWarningLevel = WarningLevel.Error;
+	}
+}
diff --git a/Engine/Plugins/VirtualProduction/CaptureManager/CaptureManagerDevices/Source/MonoVideoLiveLinkDevice/Private/MonoVideoLiveLinkDevice.cpp b/Engine/Plugins/VirtualProduction/CaptureManager/CaptureManagerDevices/Source/MonoVideoLiveLinkDevice/Private/MonoVideoLiveLinkDevice.cpp
new file mode 100644
--- /dev/null
+++ b/Engine/Plugins/VirtualProduction/CaptureManager/CaptureManagerDevices/Source/MonoVideoLiveLinkDevice/Private/MonoVideoLiveLinkDevice.cpp
@@ -0,0 +1,276 @@
+// Copyright Epic Games, Inc. All Rights Reserved.
+
+#include "MonoVideoLiveLinkDevice.h"
+
+#include "HAL/FileManager.h"
+#include "Settings/CaptureManagerSettings.h"
+
+#include "CaptureManagerMediaRWModule.h"
+
+#include "Utils/MonoVideoDeviceTokenParser.h"
+
+#include "MonoVideoLiveLinkDeviceLog.h"
+
+DEFINE_LOG_CATEGORY(LogMonoVideoLiveLinkDevice);
+
+namespace UE::MonoVideoLiveLinkDevice::Private
+{
+static const TArray<FStringView> SupportedVideoFormats =
+{
+	TEXTVIEW("mp4"),
+	TEXTVIEW("mov")
+};
+
+
+FMediaRWManager& GetMediaRWManager()
+{
+	return FModuleManager::LoadModuleChecked<FCaptureManagerMediaRWModule>("CaptureManagerMediaRW").Get();
+}
+}
+
+const UMonoVideoLiveLinkDeviceSettings* UMonoVideoLiveLinkDevice::GetSettings() const
+{
+	return CastChecked<UMonoVideoLiveLinkDeviceSettings>(Settings);
+}
+
+TSubclassOf<ULiveLinkDeviceSettings> UMonoVideoLiveLinkDevice::StaticGetSettingsClass() const
+{
+	return UMonoVideoLiveLinkDeviceSettings::StaticClass();
+}
+
+FText UMonoVideoLiveLinkDevice::GetDisplayName() const
+{
+	return FText::FromString(GetSettings()->DisplayName);
+}
+
+EDeviceHealth UMonoVideoLiveLinkDevice::GetDeviceHealth() const
+{
+	return EDeviceHealth::Nominal;
+}
+
+FText UMonoVideoLiveLinkDevice::GetHealthText() const
+{
+	return FText::FromString("Nominal");
+}
+
+void UMonoVideoLiveLinkDevice::OnSettingChanged(const FPropertyChangedEvent& InPropertyChangedEvent)
+{
+	Super::OnSettingChanged(InPropertyChangedEvent);
+
+	if (GET_MEMBER_NAME_CHECKED(UMonoVideoLiveLinkDeviceSettings, StoragePath) == InPropertyChangedEvent.GetMemberPropertyName())
+	{
+		FString Path = GetSettings()->StoragePath.Path;
+		if (!Path.IsEmpty() && FPaths::DirectoryExists(Path))
+		{
+			SetConnectionStatus(ELiveLinkDeviceConnectionStatus::Connected);
+		}
+		else
+		{
+			SetConnectionStatus(ELiveLinkDeviceConnectionStatus::Disconnected);
+		}
+	}
+}
+
+FString UMonoVideoLiveLinkDevice::GetFullTakePath(UE::CaptureManager::FTakeId InTakeId) const
+{
+	if (const FString* TakePath = FullTakePaths.Find(InTakeId))
+	{
+		return *TakePath;
+	}
+
+	return FString();
+}
+
+void UMonoVideoLiveLinkDevice::RunIngestTake(int32 InTakeId, const UIngestCapability_Options* InIngestOptions)
+{
+	AsyncTask(ENamedThreads::AnyThread, [this,
+			  TakeId = InTakeId,
+			  IngestOptions = InIngestOptions]()
+	{
+		static constexpr uint32 NumberOfTasks = 2; // Convert, Upload
+
+		using namespace UE::CaptureManager;
+		TSharedPtr<FTaskProgress> TaskProgress
+			= MakeShared<FTaskProgress>(NumberOfTasks, FTaskProgress::FProgressReporter::CreateLambda([this, TakeId](double InProgress)
+		{
+			ExecuteProcessProgressReporter(TakeId, InProgress);
+		}));
+
+		Super::IngestTake(TakeId, IngestOptions, MoveTemp(TaskProgress));
+	});
+}
+
+void UMonoVideoLiveLinkDevice::UpdateTakeList_Implementation()
+{
+	RemoveAllTakes();
+
+	FString StoragePath = GetSettings()->StoragePath.Path;
+
+	TArray<FString> SupportedVideoFiles;
+	const bool bIterationResult = IFileManager::Get().IterateDirectoryRecursively(*StoragePath, [&](const TCHAR* InFileNameOrDirectory, bool bInIsDirectory)
+	{
+		if (!bInIsDirectory && IsVideoFile(InFileNameOrDirectory))
+		{			
+			SupportedVideoFiles.Add(InFileNameOrDirectory);
+		}
+
+		return true;
+	});
+
+	for (const FString& SupportedVideoFile : SupportedVideoFiles)
+	{
+		TOptional<FTakeMetadata> TakeInfoResult = ReadTake(SupportedVideoFile);
+
+		if (TakeInfoResult.IsSet())
+		{
+			int32 TakeId = AddTake(TakeInfoResult.GetValue());
+			FullTakePaths.Add(TakeId, SupportedVideoFile);
+		}
+	}
+
+	ExecuteUpdateTakeListCallback(Execute_GetTakeIdentifiers(this));
+}
+
+ELiveLinkDeviceConnectionStatus UMonoVideoLiveLinkDevice::GetConnectionStatus_Implementation() const
+{
+	const UMonoVideoLiveLinkDeviceSettings* DeviceSettings = GetSettings();
+	FString Path = DeviceSettings->StoragePath.Path;
+
+	if (!Path.IsEmpty() && FPaths::DirectoryExists(Path))
+	{
+		return ELiveLinkDeviceConnectionStatus::Connected;
+	}
+
+	return ELiveLinkDeviceConnectionStatus::Disconnected;
+}
+
+FString UMonoVideoLiveLinkDevice::GetHardwareId_Implementation() const
+{
+	return FPlatformMisc::GetDeviceId();
+}
+
+bool UMonoVideoLiveLinkDevice::SetHardwareId_Implementation(const FString& HardwareID)
+{
+	return false;
+}
+
+bool UMonoVideoLiveLinkDevice::Connect_Implementation()
+{
+	const UMonoVideoLiveLinkDeviceSettings* DeviceSettings = GetSettings();
+	FString Path = DeviceSettings->StoragePath.Path;
+
+	if (!Path.IsEmpty() && FPaths::DirectoryExists(Path))
+	{
+		SetConnectionStatus(ELiveLinkDeviceConnectionStatus::Connected);
+		return true;
+	}
+
+	return false;
+}
+
+bool UMonoVideoLiveLinkDevice::Disconnect_Implementation()
+{
+	SetConnectionStatus(ELiveLinkDeviceConnectionStatus::Disconnected);
+	return true;
+}
+
+bool UMonoVideoLiveLinkDevice::IsVideoFile(const FString& InFileNameWithExtension)
+{
+	FString Extension = FPaths::GetExtension(InFileNameWithExtension);
+
+	return UE::MonoVideoLiveLinkDevice::Private::SupportedVideoFormats.Contains(Extension);
+}
+
+TOptional<FTakeMetadata> UMonoVideoLiveLinkDevice::ReadTake(const FString& InCurrentTakeFile) const
+{
+	using namespace UE::CaptureManager;
+
+	FMediaRWManager& MediaManager = UE::MonoVideoLiveLinkDevice::Private::GetMediaRWManager();
+
+	FTakeMetadata TakeMetadata;
+
+	FFileStatData FileData = IFileManager::Get().GetStatData(*InCurrentTakeFile);
+
+	FString FileNameFormat = GetSettings()->FileNameFormat;
+	FString FileName = FPaths::GetBaseFilename(InCurrentTakeFile);
+
+	FString SlateName;
+	FString UserId;
+	int32 TakeNumber = INDEX_NONE;
+
+	if (!FileNameFormat.IsEmpty())
+	{
+		FMonoVideoDeviceTokenParser TokenParser(FileNameFormat, FileName);
+		bool bParse = TokenParser.Parse();
+
+		if (!bParse)
+		{
+			UE_LOG(LogMonoVideoLiveLinkDevice, Warning, TEXT("Failed to match the specified format (%s) with the video file (%s)"), *(FileNameFormat), *InCurrentTakeFile);
+			return {};
+		}
+
+		SlateName = TokenParser.GetSlateName();
+		TakeNumber = TokenParser.GetTakeNumber();
+		UserId = TokenParser.GetUserId();
+	}
+
+	if (SlateName.IsEmpty())
+	{
+		SlateName = MoveTemp(FileName);
+	}
+
+	if (TakeNumber == INDEX_NONE)
+	{
+		TakeNumber = 1;
+	}
+
+	if (UserId.IsEmpty())
+	{
+		UserId = TEXT("video");
+	}
+
+	TakeMetadata.Version.Major = 4;
+	TakeMetadata.Version.Minor = 1;
+
+	TakeMetadata.Slate = MoveTemp(SlateName);
+	TakeMetadata.TakeNumber = TakeNumber;
+
+	// Create new guid for old takes
+	TakeMetadata.UniqueId = FGuid::NewGuid().ToString(EGuidFormats::DigitsWithHyphensLower);
+	TakeMetadata.DateTime = FileData.CreationTime;
+	TakeMetadata.Device.Model = TEXT("MonoVideo");
+	
+	FTakeMetadata::FVideo Video;
+	TValueOrError<TUniquePtr<IVideoReader>, FText> VideoReaderResult = MediaManager.CreateVideoReader(InCurrentTakeFile);
+	if (VideoReaderResult.HasValue())
+	{
+		Video.FrameRate = static_cast<float>(VideoReaderResult.StealValue()->GetFrameRate().AsDecimal());
+	}
+	else
+	{
+		UE_LOG(LogMonoVideoLiveLinkDevice, Warning, TEXT("Failed to determine the frame rate for the video file %s"), *InCurrentTakeFile);
+	}
+
+	Video.UserId = MoveTemp(UserId);
+	Video.Format = FPaths::GetExtension(InCurrentTakeFile);
+	Video.Path = InCurrentTakeFile;
+	Video.PathType = FTakeMetadata::FVideo::EPathType::File;
+
+	TakeMetadata.Video.Add(MoveTemp(Video));
+
+	// Check if the video file contains audio stream by creating a audio reader from the file
+	// If creation fails, we assume that there is no audio in the file
+	TValueOrError<TUniquePtr<IAudioReader>, FText> AudioReaderResult = MediaManager.CreateAudioReader(InCurrentTakeFile);
+	
+	if (AudioReaderResult.HasValue())
+	{
+		FTakeMetadata::FAudio Audio;
+		Audio.UserId = TEXT("audio");
+		Audio.Path = InCurrentTakeFile;
+		Audio.Duration = static_cast<float>(AudioReaderResult.StealValue()->GetDuration().GetTotalSeconds());
+
+		TakeMetadata.Audio.Add(MoveTemp(Audio));
+	}
+
+	return TakeMetadata;
+}
diff --git a/Engine/Plugins/VirtualProduction/CaptureManager/CaptureManagerDevices/Source/MonoVideoLiveLinkDevice/Private/MonoVideoLiveLinkDevice.h b/Engine/Plugins/VirtualProduction/CaptureManager/CaptureManagerDevices/Source/MonoVideoLiveLinkDevice/Private/MonoVideoLiveLinkDevice.h
new file mode 100644
--- /dev/null
+++ b/Engine/Plugins/VirtualProduction/CaptureManager/CaptureManagerDevices/Source/MonoVideoLiveLinkDevice/Private/MonoVideoLiveLinkDevice.h
@@ -0,0 +1,89 @@
+// Copyright Epic Games, Inc. All Rights Reserved.
+
+#pragma once
+
+#include "LiveLinkDevice.h"
+
+#include "Ingest/LiveLinkDeviceCapability_Ingest.h"
+#include "LiveLinkDeviceCapability_Connection.h"
+
+#include "Misc/Optional.h"
+
+#include "Async/TaskProgress.h"
+#include "LiveLinkHubExportManager.h"
+
+#include "BaseIngestLiveLinkDevice.h"
+
+#include "MonoVideoLiveLinkDevice.generated.h"
+
+
+UCLASS(BlueprintType)
+class UMonoVideoLiveLinkDeviceSettings : public ULiveLinkDeviceSettings
+{
+	GENERATED_BODY()
+
+public:
+	UPROPERTY(EditAnywhere, BlueprintReadWrite, Category = "Mono Video Live Link Device")
+	FString DisplayName = TEXT("Mono Video Live Link Device");
+
+	UPROPERTY(EditAnywhere, BlueprintReadWrite, Category = "Mono Video Live Link Device")
+	FDirectoryPath StoragePath;
+
+	/** 
+	 * Format of the file name that the device expects
+	 *		Available tokens:
+	 *			* {SlateName}
+	 *			* {TakeNumber}
+	 *			* {UserId}
+	 *			* {Any}
+	 *		Available delimiters: '_', '-', '.' 
+	 * 
+	 * Leave empty to parse every video file in the specified directory
+	 */
+	UPROPERTY(EditAnywhere, BlueprintReadWrite, Category = "Mono Video Live Link Device")
+	FString FileNameFormat = TEXT("{SlateName}_{TakeNumber}_{Any}");
+};
+
+
+UCLASS(BlueprintType)
+class UMonoVideoLiveLinkDevice final : public UBaseIngestLiveLinkDevice
+	, public ILiveLinkDeviceCapability_Connection
+{
+	GENERATED_BODY()
+
+public:
+
+	UFUNCTION(BlueprintCallable, Category = "GenericArchiveDeviceSettings")
+	const UMonoVideoLiveLinkDeviceSettings* GetSettings() const;
+
+	//~ Begin ULiveLinkDevice interface
+	virtual TSubclassOf<ULiveLinkDeviceSettings> StaticGetSettingsClass() const override;
+	virtual FText GetDisplayName() const override;
+	virtual EDeviceHealth GetDeviceHealth() const override;
+	virtual FText GetHealthText() const override;
+
+	virtual void OnSettingChanged(const FPropertyChangedEvent& InPropertyChangedEvent) override;
+	//~ End ULiveLinkDevice interface
+
+private:
+	virtual FString GetFullTakePath(UE::CaptureManager::FTakeId InTakeId) const override;
+
+	//~ Begin ULiveLinkDeviceCapability_Ingest interface
+	virtual void RunIngestTake(int32 InTakeId, const UIngestCapability_Options* InIngestOptions) override;
+	virtual void UpdateTakeList_Implementation() override;
+	//~ End ULiveLinkDeviceCapability_Ingest interface
+
+	//~ Begin ILiveLinkDeviceCapability_Connection interface
+	virtual ELiveLinkDeviceConnectionStatus GetConnectionStatus_Implementation() const override;
+	virtual FString GetHardwareId_Implementation() const override;
+	virtual bool SetHardwareId_Implementation(const FString& HardwareID) override;
+	virtual bool Connect_Implementation() override;
+	virtual bool Disconnect_Implementation() override;
+	//~ End ILiveLinkDeviceCapability_Connection interface
+
+	bool IsVideoFile(const FString& InFileNameWithExtension);
+
+	TOptional<FTakeMetadata> ReadTake(const FString& InCurrentVideoFile) const;
+
+	TMap<int32, FString> FullTakePaths;
+};
diff --git a/Engine/Plugins/VirtualProduction/CaptureManager/CaptureManagerDevices/Source/MonoVideoLiveLinkDevice/Private/MonoVideoLiveLinkDeviceLog.h b/Engine/Plugins/VirtualProduction/CaptureManager/CaptureManagerDevices/Source/MonoVideoLiveLinkDevice/Private/MonoVideoLiveLinkDeviceLog.h
new file mode 100644
--- /dev/null
+++ b/Engine/Plugins/VirtualProduction/CaptureManager/CaptureManagerDevices/Source/MonoVideoLiveLinkDevice/Private/MonoVideoLiveLinkDeviceLog.h
@@ -0,0 +1,7 @@
+// Copyright Epic Games, Inc. All Rights Reserved.
+
+#pragma once
+
+#include "Logging/LogMacros.h"
+
+DECLARE_LOG_CATEGORY_EXTERN(LogMonoVideoLiveLinkDevice, Log, All);
\ No newline at end of file
diff --git a/Engine/Plugins/VirtualProduction/CaptureManager/CaptureManagerDevices/Source/MonoVideoLiveLinkDevice/Private/MonoVideoLiveLinkDeviceModule.cpp b/Engine/Plugins/VirtualProduction/CaptureManager/CaptureManagerDevices/Source/MonoVideoLiveLinkDevice/Private/MonoVideoLiveLinkDeviceModule.cpp
new file mode 100644
--- /dev/null
+++ b/Engine/Plugins/VirtualProduction/CaptureManager/CaptureManagerDevices/Source/MonoVideoLiveLinkDevice/Private/MonoVideoLiveLinkDeviceModule.cpp
@@ -0,0 +1,5 @@
+// Copyright Epic Games, Inc. All Rights Reserved.
+
+#include "Modules/ModuleManager.h"
+
+IMPLEMENT_MODULE(FDefaultModuleImpl, MonoVideoLiveLinkDevice)
\ No newline at end of file
diff --git a/Engine/Plugins/VirtualProduction/CaptureManager/CaptureManagerDevices/Source/MonoVideoLiveLinkDevice/Private/Tests/TestMonoVideoDeviceTokenParser.cpp b/Engine/Plugins/VirtualProduction/CaptureManager/CaptureManagerDevices/Source/MonoVideoLiveLinkDevice/Private/Tests/TestMonoVideoDeviceTokenParser.cpp
new file mode 100644
--- /dev/null
+++ b/Engine/Plugins/VirtualProduction/CaptureManager/CaptureManagerDevices/Source/MonoVideoLiveLinkDevice/Private/Tests/TestMonoVideoDeviceTokenParser.cpp
@@ -0,0 +1,424 @@
+// Copyright Epic Games, Inc. All Rights Reserved.
+
+#include "Utils/MonoVideoDeviceTokenParser.h"
+
+#include "Misc/AutomationTest.h"
+
+IMPLEMENT_SIMPLE_AUTOMATION_TEST(FMonoVideoDeviceTokenParserTest1, "CaptureManager.MonoVideoLiveLinkDevice.TokenParser.Success",
+								 EAutomationTestFlags::EditorContext | EAutomationTestFlags::EngineFilter);
+
+bool FMonoVideoDeviceTokenParserTest1::RunTest(const FString& Parameters)
+{
+	FString InputFormat = TEXT("{SlateName}_{TakeNumber}_{UserId}");
+
+	FString SlateName = TEXT("SlateName");
+	int32 TakeNumber = 10;
+	FString UserId = TEXT("UserId");
+
+	FString FormattedValue = FString::Format(TEXT("{0}_{1}_{2}"), { SlateName, TakeNumber, UserId });
+
+	FMonoVideoDeviceTokenParser Parser(MoveTemp(InputFormat), MoveTemp(FormattedValue));
+	TestTrue(TEXT("Parse succeeded"), Parser.Parse());
+	TestEqual(TEXT("Slate Name matches"), SlateName, Parser.GetSlateName());
+	TestEqual(TEXT("Take Number matches"), TakeNumber, Parser.GetTakeNumber());
+	TestEqual(TEXT("User ID matches"), UserId, Parser.GetUserId());
+
+	return true;
+}
+
+IMPLEMENT_SIMPLE_AUTOMATION_TEST(FMonoVideoDeviceTokenParserTest2, "CaptureManager.MonoVideoLiveLinkDevice.TokenParser.Fail",
+								 EAutomationTestFlags::EditorContext | EAutomationTestFlags::EngineFilter);
+
+bool FMonoVideoDeviceTokenParserTest2::RunTest(const FString& Parameters)
+{
+	FString InputFormat = TEXT("{TakeNumber}_{UserId}");
+
+	FString SlateName = TEXT("SlateName");
+	int32 TakeNumber = 10;
+	FString UserId = TEXT("UserId");
+
+	FString FormattedValue = FString::Format(TEXT("{0}_{1}_{2}"), { SlateName, TakeNumber, UserId });
+
+	FMonoVideoDeviceTokenParser Parser(MoveTemp(InputFormat), MoveTemp(FormattedValue));
+	TestFalse(TEXT("Parse failed"), Parser.Parse());
+
+	return true;
+}
+
+IMPLEMENT_SIMPLE_AUTOMATION_TEST(FMonoVideoDeviceTokenParserTest3, "CaptureManager.MonoVideoLiveLinkDevice.TokenParser.Postfix.SuccessAny",
+								 EAutomationTestFlags::EditorContext | EAutomationTestFlags::EngineFilter);
+
+bool FMonoVideoDeviceTokenParserTest3::RunTest(const FString& Parameters)
+{
+	FString InputFormat = TEXT("{SlateName}_{UserId}_{Any}");
+
+	FString SlateName = TEXT("SlateName");
+	FString UserId = TEXT("UserId");
+
+	FString FormattedValue = FString::Format(TEXT("{0}_{1}_TestingAnythingHere"), { SlateName, UserId });
+
+	FMonoVideoDeviceTokenParser Parser(MoveTemp(InputFormat), MoveTemp(FormattedValue));
+	TestTrue(TEXT("Parse succeeded"), Parser.Parse());
+	TestEqual(TEXT("Slate Name matches"), SlateName, Parser.GetSlateName());
+	TestEqual(TEXT("User ID matches"), UserId, Parser.GetUserId());
+
+	return true;
+}
+
+IMPLEMENT_SIMPLE_AUTOMATION_TEST(FMonoVideoDeviceTokenParserTest4, "CaptureManager.MonoVideoLiveLinkDevice.TokenParser.Postfix.FailAny",
+								 EAutomationTestFlags::EditorContext | EAutomationTestFlags::EngineFilter);
+
+bool FMonoVideoDeviceTokenParserTest4::RunTest(const FString& Parameters)
+{
+	FString InputFormat = TEXT("{SlateName}_{UserId}_");
+
+	FString SlateName = TEXT("SlateName");
+	FString UserId = TEXT("UserId");
+
+	FString FormattedValue = FString::Format(TEXT("{0}_{1}_TestingAnythingHere"), { SlateName, UserId });
+
+	FMonoVideoDeviceTokenParser Parser(MoveTemp(InputFormat), MoveTemp(FormattedValue));
+	TestFalse(TEXT("Parse failed"), Parser.Parse());
+
+	return true;
+}
+
+IMPLEMENT_SIMPLE_AUTOMATION_TEST(FMonoVideoDeviceTokenParserTest5, "CaptureManager.MonoVideoLiveLinkDevice.TokenParser.Postfix.SuccessCustom",
+								 EAutomationTestFlags::EditorContext | EAutomationTestFlags::EngineFilter);
+
+bool FMonoVideoDeviceTokenParserTest5::RunTest(const FString& Parameters)
+{
+	FString CustomString = TEXT("CustomString");
+
+	FString InputFormat = FString::Format(TEXT("{SlateName}_{UserId}_{0}"), { CustomString });
+
+	FString SlateName = TEXT("SlateName");
+	FString UserId = TEXT("UserId");
+
+	FString FormattedValue = FString::Format(TEXT("{0}_{1}_{2}"), { SlateName, UserId, CustomString });
+
+	FMonoVideoDeviceTokenParser Parser(MoveTemp(InputFormat), MoveTemp(FormattedValue));
+	TestTrue(TEXT("Parse succeeded"), Parser.Parse());
+	TestEqual(TEXT("Slate Name matches"), SlateName, Parser.GetSlateName());
+	TestEqual(TEXT("User ID matches"), UserId, Parser.GetUserId());
+
+	return true;
+}
+
+IMPLEMENT_SIMPLE_AUTOMATION_TEST(FMonoVideoDeviceTokenParserTest6, "CaptureManager.MonoVideoLiveLinkDevice.TokenParser.Postfix.FailCustom",
+								 EAutomationTestFlags::EditorContext | EAutomationTestFlags::EngineFilter);
+
+bool FMonoVideoDeviceTokenParserTest6::RunTest(const FString& Parameters)
+{
+	FString CustomString = TEXT("CustomString");
+
+	FString InputFormat = FString::Format(TEXT("{SlateName}_{UserId}_{0}"), { CustomString });
+
+	FString SlateName = TEXT("SlateName");
+	FString UserId = TEXT("UserId");
+
+	FString InvalidCustomString = TEXT("InvalidCustomString");
+	FString FormattedValue = FString::Format(TEXT("{0}_{1}_{2}"), { SlateName, UserId, InvalidCustomString });
+
+	FMonoVideoDeviceTokenParser Parser(MoveTemp(InputFormat), MoveTemp(FormattedValue));
+	TestFalse(TEXT("Parse fail"), Parser.Parse());
+
+	return true;
+}
+
+IMPLEMENT_SIMPLE_AUTOMATION_TEST(FMonoVideoDeviceTokenParserTest7, "CaptureManager.MonoVideoLiveLinkDevice.TokenParser.Prefix.SuccessAny",
+								 EAutomationTestFlags::EditorContext | EAutomationTestFlags::EngineFilter);
+
+bool FMonoVideoDeviceTokenParserTest7::RunTest(const FString& Parameters)
+{
+	FString InputFormat = TEXT("{Any}_{SlateName}_{UserId}");
+
+	FString SlateName = TEXT("SlateName");
+	FString UserId = TEXT("UserId");
+
+	FString FormattedValue = FString::Format(TEXT("TestingAnythingHere_{0}_{1}"), { SlateName, UserId });
+
+	FMonoVideoDeviceTokenParser Parser(MoveTemp(InputFormat), MoveTemp(FormattedValue));
+	TestTrue(TEXT("Parse succeeded"), Parser.Parse());
+	TestEqual(TEXT("Slate Name matches"), SlateName, Parser.GetSlateName());
+	TestEqual(TEXT("User ID matches"), UserId, Parser.GetUserId());
+
+	return true;
+}
+
+IMPLEMENT_SIMPLE_AUTOMATION_TEST(FMonoVideoDeviceTokenParserTest8, "CaptureManager.MonoVideoLiveLinkDevice.TokenParser.Prefix.FailAny",
+								 EAutomationTestFlags::EditorContext | EAutomationTestFlags::EngineFilter);
+
+bool FMonoVideoDeviceTokenParserTest8::RunTest(const FString& Parameters)
+{
+	FString InputFormat = TEXT("_{SlateName}_{UserId}");
+
+	FString SlateName = TEXT("SlateName");
+	FString UserId = TEXT("UserId");
+
+	FString FormattedValue = FString::Format(TEXT("TestingAnythingHere_{0}_{1}"), { SlateName, UserId });
+
+	FMonoVideoDeviceTokenParser Parser(MoveTemp(InputFormat), MoveTemp(FormattedValue));
+	TestFalse(TEXT("Parse failed"), Parser.Parse());
+
+	return true;
+}
+
+IMPLEMENT_SIMPLE_AUTOMATION_TEST(FMonoVideoDeviceTokenParserTest9, "CaptureManager.MonoVideoLiveLinkDevice.TokenParser.Prefix.SuccessCustom",
+								 EAutomationTestFlags::EditorContext | EAutomationTestFlags::EngineFilter);
+
+bool FMonoVideoDeviceTokenParserTest9::RunTest(const FString& Parameters)
+{
+	FString CustomString = TEXT("CustomString");
+
+	FString InputFormat = FString::Format(TEXT("{0}_{SlateName}_{UserId}"), { CustomString });
+
+	FString SlateName = TEXT("SlateName");
+	FString UserId = TEXT("UserId");
+
+	FString FormattedValue = FString::Format(TEXT("{0}_{1}_{2}"), { CustomString, SlateName, UserId });
+
+	FMonoVideoDeviceTokenParser Parser(MoveTemp(InputFormat), MoveTemp(FormattedValue));
+	TestTrue(TEXT("Parse succeeded"), Parser.Parse());
+	TestEqual(TEXT("Slate Name matches"), SlateName, Parser.GetSlateName());
+	TestEqual(TEXT("User ID matches"), UserId, Parser.GetUserId());
+
+	return true;
+}
+
+IMPLEMENT_SIMPLE_AUTOMATION_TEST(FMonoVideoDeviceTokenParserTest10, "CaptureManager.MonoVideoLiveLinkDevice.TokenParser.Prefix.FailCustom",
+								 EAutomationTestFlags::EditorContext | EAutomationTestFlags::EngineFilter);
+
+bool FMonoVideoDeviceTokenParserTest10::RunTest(const FString& Parameters)
+{
+	FString CustomString = TEXT("CustomString");
+
+	FString InputFormat = FString::Format(TEXT("{0}_{SlateName}_{UserId}"), { CustomString });
+
+	FString SlateName = TEXT("SlateName");
+	FString UserId = TEXT("UserId");
+
+	FString InvalidCustomString = TEXT("InvalidCustomString");
+	FString FormattedValue = FString::Format(TEXT("{0}_{1}_{2}"), { InvalidCustomString, SlateName, UserId });
+
+	FMonoVideoDeviceTokenParser Parser(MoveTemp(InputFormat), MoveTemp(FormattedValue));
+	TestFalse(TEXT("Parse fail"), Parser.Parse());
+
+	return true;
+}
+
+IMPLEMENT_SIMPLE_AUTOMATION_TEST(FMonoVideoDeviceTokenParserTest11, "CaptureManager.MonoVideoLiveLinkDevice.TokenParser.Mid.SuccessAny",
+								 EAutomationTestFlags::EditorContext | EAutomationTestFlags::EngineFilter);
+
+bool FMonoVideoDeviceTokenParserTest11::RunTest(const FString& Parameters)
+{
+	FString InputFormat = TEXT("{SlateName}_{Any}_{UserId}");
+
+	FString SlateName = TEXT("SlateName");
+	FString UserId = TEXT("UserId");
+
+	FString FormattedValue = FString::Format(TEXT("{0}_TestingAnythingHere_{1}"), { SlateName, UserId });
+
+	FMonoVideoDeviceTokenParser Parser(MoveTemp(InputFormat), MoveTemp(FormattedValue));
+	TestTrue(TEXT("Parse succeeded"), Parser.Parse());
+	TestEqual(TEXT("Slate Name matches"), SlateName, Parser.GetSlateName());
+	TestEqual(TEXT("User ID matches"), UserId, Parser.GetUserId());
+
+	return true;
+}
+
+IMPLEMENT_SIMPLE_AUTOMATION_TEST(FMonoVideoDeviceTokenParserTest12, "CaptureManager.MonoVideoLiveLinkDevice.TokenParser.Mid.FailAny",
+								 EAutomationTestFlags::EditorContext | EAutomationTestFlags::EngineFilter);
+
+bool FMonoVideoDeviceTokenParserTest12::RunTest(const FString& Parameters)
+{
+	FString InputFormat = TEXT("{SlateName}_{UserId}");
+
+	FString SlateName = TEXT("SlateName");
+	FString UserId = TEXT("UserId");
+
+	FString FormattedValue = FString::Format(TEXT("{0}_TestingAnythingHere_{1}"), { SlateName, UserId });
+
+	FMonoVideoDeviceTokenParser Parser(MoveTemp(InputFormat), MoveTemp(FormattedValue));
+	TestFalse(TEXT("Parse failed"), Parser.Parse());
+
+	return true;
+}
+
+IMPLEMENT_SIMPLE_AUTOMATION_TEST(FMonoVideoDeviceTokenParserTest13, "CaptureManager.MonoVideoLiveLinkDevice.TokenParser.Mid.SuccessCustom",
+								 EAutomationTestFlags::EditorContext | EAutomationTestFlags::EngineFilter);
+
+bool FMonoVideoDeviceTokenParserTest13::RunTest(const FString& Parameters)
+{
+	FString CustomString = TEXT("CustomString");
+
+	FString InputFormat = FString::Format(TEXT("{SlateName}_{0}_{UserId}"), { CustomString });
+
+	FString SlateName = TEXT("SlateName");
+	FString UserId = TEXT("UserId");
+
+	FString FormattedValue = FString::Format(TEXT("{0}_{1}_{2}"), { SlateName, CustomString, UserId });
+
+	FMonoVideoDeviceTokenParser Parser(MoveTemp(InputFormat), MoveTemp(FormattedValue));
+	TestTrue(TEXT("Parse succeeded"), Parser.Parse());
+	TestEqual(TEXT("Slate Name matches"), SlateName, Parser.GetSlateName());
+	TestEqual(TEXT("User ID matches"), UserId, Parser.GetUserId());
+
+	return true;
+}
+
+IMPLEMENT_SIMPLE_AUTOMATION_TEST(FMonoVideoDeviceTokenParserTest14, "CaptureManager.MonoVideoLiveLinkDevice.TokenParser.Mid.FailCustom",
+								 EAutomationTestFlags::EditorContext | EAutomationTestFlags::EngineFilter);
+
+bool FMonoVideoDeviceTokenParserTest14::RunTest(const FString& Parameters)
+{
+	FString CustomString = TEXT("CustomString");
+
+	FString InputFormat = FString::Format(TEXT("{SlateName}_{0}_{UserId}"), { CustomString });
+
+	FString SlateName = TEXT("SlateName");
+	FString UserId = TEXT("UserId");
+
+	FString InvalidCustomString = TEXT("InvalidCustomString");
+	FString FormattedValue = FString::Format(TEXT("{0}_{1}_{2}"), { SlateName, InvalidCustomString, UserId });
+
+	FMonoVideoDeviceTokenParser Parser(MoveTemp(InputFormat), MoveTemp(FormattedValue));
+	TestFalse(TEXT("Parse fail"), Parser.Parse());
+
+	return true;
+}
+
+IMPLEMENT_SIMPLE_AUTOMATION_TEST(FMonoVideoDeviceTokenParserTest15, "CaptureManager.MonoVideoLiveLinkDevice.TokenParser.DifferentDelimiterFailure",
+								 EAutomationTestFlags::EditorContext | EAutomationTestFlags::EngineFilter);
+
+bool FMonoVideoDeviceTokenParserTest15::RunTest(const FString& Parameters)
+{
+	FString InputFormat = TEXT("{SlateName}_{TakeNumber}-{UserId}");
+
+	FString SlateName = TEXT("SlateName");
+	int32 TakeNumber = 10;
+	FString UserId = TEXT("UserId");
+
+	FString FormattedValue = FString::Format(TEXT("{0}_{1}-{2}"), { SlateName, TakeNumber, UserId });
+
+	FMonoVideoDeviceTokenParser Parser(MoveTemp(InputFormat), MoveTemp(FormattedValue));
+	TestFalse(TEXT("Parse failed"), Parser.Parse());
+
+	return true;
+}
+
+IMPLEMENT_SIMPLE_AUTOMATION_TEST(FMonoVideoDeviceTokenParserTest16, "CaptureManager.MonoVideoLiveLinkDevice.TokenParser.DifferentDelimiterSuccess",
+								 EAutomationTestFlags::EditorContext | EAutomationTestFlags::EngineFilter);
+
+bool FMonoVideoDeviceTokenParserTest16::RunTest(const FString& Parameters)
+{
+	FString InputFormat = TEXT("{SlateName}-{UserId}");
+
+	FString SlateName = TEXT("SlateName");
+	int32 TakeNumber = 10;
+	FString UserId = TEXT("UserId");
+
+	FString FormattedValue = FString::Format(TEXT("{0}_{1}-{2}"), { SlateName, TakeNumber, UserId });
+
+	FMonoVideoDeviceTokenParser Parser(MoveTemp(InputFormat), MoveTemp(FormattedValue));
+	TestTrue(TEXT("Parse succeeded"), Parser.Parse());
+	TestNotEqual(TEXT("Slate Name doesn't match"), SlateName, Parser.GetSlateName());
+	TestEqual(TEXT("User ID matches"), UserId, Parser.GetUserId());
+	
+	FString ExpectedSlateName = FString::Format(TEXT("{0}_{1}"), { SlateName, TakeNumber });
+	TestEqual(TEXT("Slate name matches"), ExpectedSlateName, Parser.GetSlateName());
+
+	return true;
+}
+
+IMPLEMENT_SIMPLE_AUTOMATION_TEST(FMonoVideoDeviceTokenParserTest17, "CaptureManager.MonoVideoLiveLinkDevice.TokenParser.SingleTokenValue",
+								 EAutomationTestFlags::EditorContext | EAutomationTestFlags::EngineFilter);
+
+bool FMonoVideoDeviceTokenParserTest17::RunTest(const FString& Parameters)
+{
+	FString InputFormat = TEXT("{SlateName}");
+
+	FString SlateName = TEXT("SlateName");
+	int32 TakeNumber = 10;
+
+	FString FormattedValue = FString::Format(TEXT("{0}"), { SlateName });
+
+	FMonoVideoDeviceTokenParser Parser(MoveTemp(InputFormat), MoveTemp(FormattedValue));
+	TestTrue(TEXT("Parse succeeded"), Parser.Parse());
+	TestEqual(TEXT("Slate Name matches"), SlateName, Parser.GetSlateName());
+
+	return true;
+}
+
+IMPLEMENT_SIMPLE_AUTOMATION_TEST(FMonoVideoDeviceTokenParserTest18, "CaptureManager.MonoVideoLiveLinkDevice.TokenParser.SingleValue",
+								 EAutomationTestFlags::EditorContext | EAutomationTestFlags::EngineFilter);
+
+bool FMonoVideoDeviceTokenParserTest18::RunTest(const FString& Parameters)
+{
+	FString InputFormat = TEXT("TestingString");
+	FString TestingString = TEXT("TestingString");
+
+	int32 NoTakeNumber = INDEX_NONE;
+
+	FMonoVideoDeviceTokenParser Parser(MoveTemp(InputFormat), MoveTemp(TestingString));
+	TestTrue(TEXT("Parse succeeded"), Parser.Parse());
+	TestTrue(TEXT("Slate Name empty"), Parser.GetSlateName().IsEmpty());
+	TestEqual(TEXT("Take Number invalid"), NoTakeNumber, Parser.GetTakeNumber());
+	TestTrue(TEXT("User Id empty"), Parser.GetUserId().IsEmpty());
+
+	return true;
+}
+
+IMPLEMENT_SIMPLE_AUTOMATION_TEST(FMonoVideoDeviceTokenParserTest19, "CaptureManager.MonoVideoLiveLinkDevice.TokenParser.EmptyString",
+								 EAutomationTestFlags::EditorContext | EAutomationTestFlags::EngineFilter);
+
+bool FMonoVideoDeviceTokenParserTest19::RunTest(const FString& Parameters)
+{
+	FString InputFormat = TEXT("");
+
+	FString SlateName = TEXT("SlateName");
+	int32 TakeNumber = 10;
+
+	FString FormattedValue = FString::Format(TEXT("{0}_{1}"), { SlateName, TakeNumber });
+
+	FMonoVideoDeviceTokenParser Parser(MoveTemp(InputFormat), MoveTemp(FormattedValue));
+	TestFalse(TEXT("Parse failed"), Parser.Parse());
+
+	return true;
+}
+
+IMPLEMENT_SIMPLE_AUTOMATION_TEST(FMonoVideoDeviceTokenParserTest20, "CaptureManager.MonoVideoLiveLinkDevice.TokenParser.MultipleDelimiters",
+								 EAutomationTestFlags::EditorContext | EAutomationTestFlags::EngineFilter);
+
+bool FMonoVideoDeviceTokenParserTest20::RunTest(const FString& Parameters)
+{
+	FString InputFormat = TEXT("______");
+
+	FString SlateName = TEXT("SlateName");
+	int32 TakeNumber = 10;
+
+	FString FormattedValue = FString::Format(TEXT("{0}_{1}"), { SlateName, TakeNumber });
+
+	FMonoVideoDeviceTokenParser Parser(MoveTemp(InputFormat), MoveTemp(FormattedValue));
+	TestFalse(TEXT("Parse failed"), Parser.Parse());
+
+	return true;
+}
+
+IMPLEMENT_SIMPLE_AUTOMATION_TEST(FMonoVideoDeviceTokenParserTest21, "CaptureManager.MonoVideoLiveLinkDevice.TokenParser.SingleDelimiter",
+								 EAutomationTestFlags::EditorContext | EAutomationTestFlags::EngineFilter);
+
+bool FMonoVideoDeviceTokenParserTest21::RunTest(const FString& Parameters)
+{
+	FString InputFormat = TEXT("_");
+
+	FString SlateName = TEXT("SlateName");
+	int32 TakeNumber = 10;
+
+	FString FormattedValue = FString::Format(TEXT("{0}_{1}"), { SlateName, TakeNumber });
+
+	FMonoVideoDeviceTokenParser Parser(MoveTemp(InputFormat), MoveTemp(FormattedValue));
+	TestFalse(TEXT("Parse failed"), Parser.Parse());
+
+	return true;
+}
\ No newline at end of file
diff --git a/Engine/Plugins/VirtualProduction/CaptureManager/CaptureManagerDevices/Source/MonoVideoLiveLinkDevice/Private/Utils/MonoVideoDeviceTokenParser.cpp b/Engine/Plugins/VirtualProduction/CaptureManager/CaptureManagerDevices/Source/MonoVideoLiveLinkDevice/Private/Utils/MonoVideoDeviceTokenParser.cpp
new file mode 100644
--- /dev/null
+++ b/Engine/Plugins/VirtualProduction/CaptureManager/CaptureManagerDevices/Source/MonoVideoLiveLinkDevice/Private/Utils/MonoVideoDeviceTokenParser.cpp
@@ -0,0 +1,157 @@
+// Copyright Epic Games, Inc. All Rights Reserved.
+
+#include "MonoVideoDeviceTokenParser.h"
+
+#include "Containers/Map.h"
+#include "Containers/Array.h"
+#include "Containers/StringView.h"
+
+const FString FMonoVideoDeviceTokenParser::SlateNameToken = TEXT("{SlateName}");
+const FString FMonoVideoDeviceTokenParser::TakeNumberToken = TEXT("{TakeNumber}");
+const FString FMonoVideoDeviceTokenParser::UserIDToken = TEXT("{UserId}");
+const FString FMonoVideoDeviceTokenParser::AnyToken = TEXT("{Any}");
+
+namespace UE::MonoVideoLiveLinkDevice::Private
+{
+
+static const TArray<FString::ElementType> Delimiters =
+{
+	'-',
+	'_',
+	'.'
+};
+
+TOptional<FString::ElementType> FindDelimiter(const FString& InFileName)
+{
+	for (const FString::ElementType Character : InFileName)
+	{
+		if (Delimiters.Contains(Character))
+		{
+			return Character;
+		}
+	}
+
+	return {};
+}
+
+}
+
+FMonoVideoDeviceTokenParser::FMonoVideoDeviceTokenParser(const FString& InFormat, const FString& InFormattedValue)
+	: Format(InFormat)
+	, FormattedValue(InFormattedValue)
+{
+}
+
+bool FMonoVideoDeviceTokenParser::Parse()
+{
+	TOptional<FString::ElementType> Delimiter = UE::MonoVideoLiveLinkDevice::Private::FindDelimiter(Format);
+
+	TOptional<FComponentValueMap> ComponentValueMap;
+	if (Delimiter.IsSet())
+	{
+		ComponentValueMap = ParseWithDelimiter(Delimiter.GetValue());
+	}
+	else
+	{
+		ComponentValueMap = ParseWithoutDelimiter();
+	}
+
+	if (!ComponentValueMap.IsSet())
+	{
+		return false;
+	}
+
+	return DetermineValues(MoveTemp(ComponentValueMap.GetValue()));
+}
+
+TOptional<FMonoVideoDeviceTokenParser::FComponentValueMap> FMonoVideoDeviceTokenParser::ParseWithDelimiter(FString::ElementType InDelimiter) const
+{
+	TArray<FString> FileNameFormatComponents;
+
+	FString Delimiter;
+	Delimiter.AppendChar(InDelimiter);
+
+	Format.ParseIntoArray(FileNameFormatComponents, *Delimiter);
+
+	if (FileNameFormatComponents.IsEmpty())
+	{
+		return {};
+	}
+
+	TArray<FString> FileNameComponents;
+	FormattedValue.ParseIntoArray(FileNameComponents, *Delimiter);
+
+	if (FileNameComponents.IsEmpty())
+	{
+		return {};
+	}
+
+	if (FileNameFormatComponents.Num() != FileNameComponents.Num())
+	{
+		return {};
+	}
+
+	FComponentValueMap ComponentValueMap;
+	for (int32 Index = 0; Index < FileNameComponents.Num(); ++Index)
+	{
+		ComponentValueMap.Add(FileNameFormatComponents[Index], FileNameComponents[Index]);
+	}
+
+	return ComponentValueMap;
+}
+
+TOptional<FMonoVideoDeviceTokenParser::FComponentValueMap> FMonoVideoDeviceTokenParser::ParseWithoutDelimiter() const
+{
+	FComponentValueMap ComponentValueMap;
+	ComponentValueMap.Add(Format, FormattedValue);
+
+	return ComponentValueMap;
+}
+
+bool FMonoVideoDeviceTokenParser::DetermineValues(FComponentValueMap InComponentValueMap)
+{
+	if (const FString* MaybeSlateName = InComponentValueMap.Find(SlateNameToken))
+	{
+		SlateName = *MaybeSlateName;
+		InComponentValueMap.Remove(SlateNameToken);
+	}
+
+	if (const FString* MaybeTakeNumber = InComponentValueMap.Find(TakeNumberToken))
+	{
+		TakeNumber = FCString::Atoi(**MaybeTakeNumber);
+		InComponentValueMap.Remove(TakeNumberToken);
+	}
+
+	if (const FString* MaybeUserId = InComponentValueMap.Find(UserIDToken))
+	{
+		UserId = *MaybeUserId;
+		InComponentValueMap.Remove(UserIDToken);
+	}
+
+	InComponentValueMap.Remove(AnyToken);
+
+	for (const TPair<FString, FString>& AdditionalComponentValue : InComponentValueMap)
+	{
+		if (AdditionalComponentValue.Key != AdditionalComponentValue.Value)
+		{
+			return false;
+		}
+	}
+
+	return true;
+}
+
+FString FMonoVideoDeviceTokenParser::GetSlateName() const
+{
+	return SlateName;
+}
+
+int32 FMonoVideoDeviceTokenParser::GetTakeNumber() const
+{
+	return TakeNumber;
+}
+
+FString FMonoVideoDeviceTokenParser::GetUserId() const
+{
+	return UserId;
+}
diff --git a/Engine/Plugins/VirtualProduction/CaptureManager/CaptureManagerDevices/Source/MonoVideoLiveLinkDevice/Private/Utils/MonoVideoDeviceTokenParser.h b/Engine/Plugins/VirtualProduction/CaptureManager/CaptureManagerDevices/Source/MonoVideoLiveLinkDevice/Private/Utils/MonoVideoDeviceTokenParser.h
new file mode 100644
--- /dev/null
+++ b/Engine/Plugins/VirtualProduction/CaptureManager/CaptureManagerDevices/Source/MonoVideoLiveLinkDevice/Private/Utils/MonoVideoDeviceTokenParser.h
@@ -0,0 +1,39 @@
+// Copyright Epic Games, Inc. All Rights Reserved.
+
+#pragma once
+
+#include "Containers/UnrealString.h"
+
+class FMonoVideoDeviceTokenParser
+{
+public:
+
+	static const FString SlateNameToken;
+	static const FString TakeNumberToken;
+	static const FString UserIDToken;
+	static const FString AnyToken;
+
+	FMonoVideoDeviceTokenParser(const FString& InFormat, const FString& InFormattedValue);
+
+	bool Parse();
+
+	FString GetSlateName() const;
+	int32 GetTakeNumber() const;
+	FString GetUserId() const;
+
+private:
+
+	using FComponentValueMap = TMap<FString, FString>;
+
+	TOptional<FComponentValueMap> ParseWithDelimiter(FString::ElementType InDelimiter) const;
+	TOptional<FComponentValueMap> ParseWithoutDelimiter() const;
+
+	bool DetermineValues(FComponentValueMap InComponentValueMap);
+
+	FString Format;
+	FString FormattedValue;
+
+	FString SlateName;
+	int32 TakeNumber = INDEX_NONE;
+	FString UserId;
+};
\ No newline at end of file
diff --git a/Engine/Source/Programs/LiveLinkHub/Source/LiveLinkHub.Target.cs b/Engine/Source/Programs/LiveLinkHub/Source/LiveLinkHub.Target.cs
--- a/Engine/Source/Programs/LiveLinkHub/Source/LiveLinkHub.Target.cs
+++ b/Engine/Source/Programs/LiveLinkHub/Source/LiveLinkHub.Target.cs
@@ -58,7 +58,8 @@ public class LiveLinkHubTarget : TargetRules
 
 			"LiveLinkOpenVR",
 
-			"CaptureManagerApp"
+			"CaptureManagerApp",
+			"CaptureManagerDevices"
 		});
 
 		if (bEnableMetaHumanLiveLinkPlugin)
