Move Capture Manager devices to the dedicated plugin

diff --git a/Engine/Plugins/VirtualProduction/CaptureManager/CaptureManagerApp/CaptureManagerApp.uplugin b/Engine/Plugins/VirtualProduction/CaptureManager/CaptureManagerApp/CaptureManagerApp.uplugin
--- a/Engine/Plugins/VirtualProduction/CaptureManager/CaptureManagerApp/CaptureManagerApp.uplugin
+++ b/Engine/Plugins/VirtualProduction/CaptureManager/CaptureManagerApp/CaptureManagerApp.uplugin
@@ -68,11 +68,6 @@
 			"Type": "Editor",
 			"LoadingPhase": "Default"
 		},
-		{
-			"Name": "GenericLiveLinkDevice",
-			"Type": "Editor",
-			"LoadingPhase": "Default"
-		},
 		{
 			"Name": "LiveLinkCapabilities",
 			"Type": "Editor",
diff --git a/Engine/Plugins/VirtualProduction/CaptureManager/CaptureManagerApp/Source/CaptureManagerEditor/CaptureManagerEditor.Build.cs b/Engine/Plugins/VirtualProduction/CaptureManager/CaptureManagerApp/Source/CaptureManagerEditor/CaptureManagerEditor.Build.cs
--- a/Engine/Plugins/VirtualProduction/CaptureManager/CaptureManagerApp/Source/CaptureManagerEditor/CaptureManagerEditor.Build.cs
+++ b/Engine/Plugins/VirtualProduction/CaptureManager/CaptureManagerApp/Source/CaptureManagerEditor/CaptureManagerEditor.Build.cs
@@ -18,7 +18,6 @@ public class CaptureManagerEditor : ModuleRules
 
 		PrivateDependencyModuleNames.AddRange(new string[]
 		{
-			"GenericLiveLinkDevice",
 			"InputCore",
 			"LiveLinkDevice",
 			"LiveLinkCapabilities",
diff --git a/Engine/Plugins/VirtualProduction/CaptureManager/CaptureManagerApp/Source/CaptureManagerEditor/Private/Scripting/ULiveLinkHubCaptureDevice.cpp b/Engine/Plugins/VirtualProduction/CaptureManager/CaptureManagerApp/Source/CaptureManagerEditor/Private/Scripting/ULiveLinkHubCaptureDevice.cpp
--- a/Engine/Plugins/VirtualProduction/CaptureManager/CaptureManagerApp/Source/CaptureManagerEditor/Private/Scripting/ULiveLinkHubCaptureDevice.cpp
+++ b/Engine/Plugins/VirtualProduction/CaptureManager/CaptureManagerApp/Source/CaptureManagerEditor/Private/Scripting/ULiveLinkHubCaptureDevice.cpp
@@ -5,8 +5,6 @@
 #include "LiveLinkDeviceCapability_Connection.h"
 #include "Ingest/LiveLinkDeviceCapability_Ingest.h"
 
-#include "GenericLiveLinkDeviceHelper.h"
-
 #include "Engine/Engine.h"
 
 #include "LiveLinkDeviceSubsystem.h"
@@ -289,32 +287,6 @@ TObjectPtr<UIngestCapability_ProcessResult> ULiveLinkHubCaptureDevice::FImpl::Fe
 	return UIngestCapability_ProcessResult::Success();
 }
 
-ULiveLinkHubCaptureDevice* ULiveLinkHubCaptureDeviceFactory::CreateArchiveDevice(FString InName, FString InStoragePath)
-{
-	UGenericLiveLinkDeviceHelper* Helper = NewObject<UGenericLiveLinkDeviceHelper>(this);
-
-	TObjectPtr<ULiveLinkDevice> BackgroundDevice = Helper->CreateArchiveDevice(MoveTemp(InName), MoveTemp(InStoragePath));
-
-	TObjectPtr<ULiveLinkHubCaptureDevice> Device = NewObject<ULiveLinkHubCaptureDevice>(this);
-	Device->Name = MoveTemp(InName);
-	Device->ImplPtr->BackgroundDevice = MoveTemp(BackgroundDevice);
-
-	return Device;
-}
-
-ULiveLinkHubCaptureDevice* ULiveLinkHubCaptureDeviceFactory::CreateNetworkDevice(FString InName, FString InIPAddress, int32 InPort)
-{
-	UGenericLiveLinkDeviceHelper* Helper = NewObject<UGenericLiveLinkDeviceHelper>(this);
-
-	TObjectPtr<ULiveLinkDevice> BackgroundDevice = Helper->CreateNetworkDevice(InName, MoveTemp(InIPAddress), InPort);
-
-	TObjectPtr<ULiveLinkHubCaptureDevice> Device = NewObject<ULiveLinkHubCaptureDevice>(this);
-	Device->Name = MoveTemp(InName);
-	Device->ImplPtr->BackgroundDevice = MoveTemp(BackgroundDevice);
-
-	return Device;
-}
-
 ULiveLinkHubCaptureDevice* ULiveLinkHubCaptureDeviceFactory::CreateDeviceByClass(FString InName, UClass* InDeviceClass, ULiveLinkDeviceSettings* InSettings)
 {
 	TSubclassOf<ULiveLinkDevice> DeviceClass = InDeviceClass;
diff --git a/Engine/Plugins/VirtualProduction/CaptureManager/CaptureManagerApp/Source/CaptureManagerEditor/Private/Scripting/ULiveLinkHubCaptureDevice.h b/Engine/Plugins/VirtualProduction/CaptureManager/CaptureManagerApp/Source/CaptureManagerEditor/Private/Scripting/ULiveLinkHubCaptureDevice.h
--- a/Engine/Plugins/VirtualProduction/CaptureManager/CaptureManagerApp/Source/CaptureManagerEditor/Private/Scripting/ULiveLinkHubCaptureDevice.h
+++ b/Engine/Plugins/VirtualProduction/CaptureManager/CaptureManagerApp/Source/CaptureManagerEditor/Private/Scripting/ULiveLinkHubCaptureDevice.h
@@ -47,12 +47,6 @@ class ULiveLinkHubCaptureDeviceFactory : public UObject
 
 public:
 
-	UFUNCTION(BlueprintCallable, Category = "Ingest | Creator")
-	ULiveLinkHubCaptureDevice* CreateArchiveDevice(FString InName, FString InStoragePath);
-
-	UFUNCTION(BlueprintCallable, Category = "Ingest | Creator")
-	ULiveLinkHubCaptureDevice* CreateNetworkDevice(FString InName, FString InIPAddress, int32 InPort);
-
 	UFUNCTION(BlueprintCallable, Category = "Ingest | Creator")
 	ULiveLinkHubCaptureDevice* CreateDeviceByClass(FString InName, UClass* InDeviceClass, ULiveLinkDeviceSettings* Settings);
 };
diff --git a/Engine/Plugins/VirtualProduction/CaptureManager/CaptureManagerApp/Source/GenericLiveLinkDevice/GenericLiveLinkDevice.Build.cs b/Engine/Plugins/VirtualProduction/CaptureManager/CaptureManagerApp/Source/GenericLiveLinkDevice/GenericLiveLinkDevice.Build.cs
deleted file mode 100644
diff --git a/Engine/Plugins/VirtualProduction/CaptureManager/CaptureManagerApp/Source/GenericLiveLinkDevice/Private/CPS/Customizations/DeviceIpAddressCustomization.cpp b/Engine/Plugins/VirtualProduction/CaptureManager/CaptureManagerApp/Source/GenericLiveLinkDevice/Private/CPS/Customizations/DeviceIpAddressCustomization.cpp
deleted file mode 100644
diff --git a/Engine/Plugins/VirtualProduction/CaptureManager/CaptureManagerApp/Source/GenericLiveLinkDevice/Private/CPS/Customizations/DeviceIpAddressCustomization.h b/Engine/Plugins/VirtualProduction/CaptureManager/CaptureManagerApp/Source/GenericLiveLinkDevice/Private/CPS/Customizations/DeviceIpAddressCustomization.h
deleted file mode 100644
diff --git a/Engine/Plugins/VirtualProduction/CaptureManager/CaptureManagerApp/Source/GenericLiveLinkDevice/Private/CPS/Customizations/ToggleConnectActionCustomization.cpp b/Engine/Plugins/VirtualProduction/CaptureManager/CaptureManagerApp/Source/GenericLiveLinkDevice/Private/CPS/Customizations/ToggleConnectActionCustomization.cpp
deleted file mode 100644
diff --git a/Engine/Plugins/VirtualProduction/CaptureManager/CaptureManagerApp/Source/GenericLiveLinkDevice/Private/CPS/Customizations/ToggleConnectActionCustomization.h b/Engine/Plugins/VirtualProduction/CaptureManager/CaptureManagerApp/Source/GenericLiveLinkDevice/Private/CPS/Customizations/ToggleConnectActionCustomization.h
deleted file mode 100644
diff --git a/Engine/Plugins/VirtualProduction/CaptureManager/CaptureManagerApp/Source/GenericLiveLinkDevice/Private/CPS/LiveLinkFaceDevice.cpp b/Engine/Plugins/VirtualProduction/CaptureManager/CaptureManagerApp/Source/GenericLiveLinkDevice/Private/CPS/LiveLinkFaceDevice.cpp
deleted file mode 100644
diff --git a/Engine/Plugins/VirtualProduction/CaptureManager/CaptureManagerApp/Source/GenericLiveLinkDevice/Private/CPS/LiveLinkFaceDevice.h b/Engine/Plugins/VirtualProduction/CaptureManager/CaptureManagerApp/Source/GenericLiveLinkDevice/Private/CPS/LiveLinkFaceDevice.h
deleted file mode 100644
diff --git a/Engine/Plugins/VirtualProduction/CaptureManager/CaptureManagerApp/Source/GenericLiveLinkDevice/Private/CPS/Protocol/CPSDataStream.cpp b/Engine/Plugins/VirtualProduction/CaptureManager/CaptureManagerApp/Source/GenericLiveLinkDevice/Private/CPS/Protocol/CPSDataStream.cpp
deleted file mode 100644
diff --git a/Engine/Plugins/VirtualProduction/CaptureManager/CaptureManagerApp/Source/GenericLiveLinkDevice/Private/CPS/Protocol/CPSDataStream.h b/Engine/Plugins/VirtualProduction/CaptureManager/CaptureManagerApp/Source/GenericLiveLinkDevice/Private/CPS/Protocol/CPSDataStream.h
deleted file mode 100644
diff --git a/Engine/Plugins/VirtualProduction/CaptureManager/CaptureManagerApp/Source/GenericLiveLinkDevice/Private/CPS/Protocol/CPSDevice.cpp b/Engine/Plugins/VirtualProduction/CaptureManager/CaptureManagerApp/Source/GenericLiveLinkDevice/Private/CPS/Protocol/CPSDevice.cpp
deleted file mode 100644
diff --git a/Engine/Plugins/VirtualProduction/CaptureManager/CaptureManagerApp/Source/GenericLiveLinkDevice/Private/CPS/Protocol/CPSDevice.h b/Engine/Plugins/VirtualProduction/CaptureManager/CaptureManagerApp/Source/GenericLiveLinkDevice/Private/CPS/Protocol/CPSDevice.h
deleted file mode 100644
diff --git a/Engine/Plugins/VirtualProduction/CaptureManager/CaptureManagerApp/Source/GenericLiveLinkDevice/Private/CPS/Protocol/CPSFileStream.cpp b/Engine/Plugins/VirtualProduction/CaptureManager/CaptureManagerApp/Source/GenericLiveLinkDevice/Private/CPS/Protocol/CPSFileStream.cpp
deleted file mode 100644
diff --git a/Engine/Plugins/VirtualProduction/CaptureManager/CaptureManagerApp/Source/GenericLiveLinkDevice/Private/CPS/Protocol/CPSFileStream.h b/Engine/Plugins/VirtualProduction/CaptureManager/CaptureManagerApp/Source/GenericLiveLinkDevice/Private/CPS/Protocol/CPSFileStream.h
deleted file mode 100644
diff --git a/Engine/Plugins/VirtualProduction/CaptureManager/CaptureManagerApp/Source/GenericLiveLinkDevice/Private/Filesystem/FilesystemLiveLinkDevice.cpp b/Engine/Plugins/VirtualProduction/CaptureManager/CaptureManagerApp/Source/GenericLiveLinkDevice/Private/Filesystem/FilesystemLiveLinkDevice.cpp
deleted file mode 100644
diff --git a/Engine/Plugins/VirtualProduction/CaptureManager/CaptureManagerApp/Source/GenericLiveLinkDevice/Private/Filesystem/FilesystemLiveLinkDevice.h b/Engine/Plugins/VirtualProduction/CaptureManager/CaptureManagerApp/Source/GenericLiveLinkDevice/Private/Filesystem/FilesystemLiveLinkDevice.h
deleted file mode 100644
diff --git a/Engine/Plugins/VirtualProduction/CaptureManager/CaptureManagerApp/Source/GenericLiveLinkDevice/Private/GenericLiveLinkDeviceHelper.cpp b/Engine/Plugins/VirtualProduction/CaptureManager/CaptureManagerApp/Source/GenericLiveLinkDevice/Private/GenericLiveLinkDeviceHelper.cpp
deleted file mode 100644
diff --git a/Engine/Plugins/VirtualProduction/CaptureManager/CaptureManagerApp/Source/GenericLiveLinkDevice/Private/GenericLiveLinkDeviceModule.cpp b/Engine/Plugins/VirtualProduction/CaptureManager/CaptureManagerApp/Source/GenericLiveLinkDevice/Private/GenericLiveLinkDeviceModule.cpp
deleted file mode 100644
diff --git a/Engine/Plugins/VirtualProduction/CaptureManager/CaptureManagerApp/Source/GenericLiveLinkDevice/Private/Tests/GenericLiveLinkDeviceTests.cpp b/Engine/Plugins/VirtualProduction/CaptureManager/CaptureManagerApp/Source/GenericLiveLinkDevice/Private/Tests/GenericLiveLinkDeviceTests.cpp
deleted file mode 100644
diff --git a/Engine/Plugins/VirtualProduction/CaptureManager/CaptureManagerApp/Source/GenericLiveLinkDevice/Private/Utils/GenericLiveLinkDeviceLog.h b/Engine/Plugins/VirtualProduction/CaptureManager/CaptureManagerApp/Source/GenericLiveLinkDevice/Private/Utils/GenericLiveLinkDeviceLog.h
deleted file mode 100644
diff --git a/Engine/Plugins/VirtualProduction/CaptureManager/CaptureManagerApp/Source/GenericLiveLinkDevice/Private/Utils/GenericLiveLinkDeviceUtility.cpp b/Engine/Plugins/VirtualProduction/CaptureManager/CaptureManagerApp/Source/GenericLiveLinkDevice/Private/Utils/GenericLiveLinkDeviceUtility.cpp
deleted file mode 100644
diff --git a/Engine/Plugins/VirtualProduction/CaptureManager/CaptureManagerApp/Source/GenericLiveLinkDevice/Private/Utils/GenericLiveLinkDeviceUtility.h b/Engine/Plugins/VirtualProduction/CaptureManager/CaptureManagerApp/Source/GenericLiveLinkDevice/Private/Utils/GenericLiveLinkDeviceUtility.h
deleted file mode 100644
diff --git a/Engine/Plugins/VirtualProduction/CaptureManager/CaptureManagerApp/Source/GenericLiveLinkDevice/Public/GenericLiveLinkDeviceHelper.h b/Engine/Plugins/VirtualProduction/CaptureManager/CaptureManagerApp/Source/GenericLiveLinkDevice/Public/GenericLiveLinkDeviceHelper.h
deleted file mode 100644
diff --git a/Engine/Plugins/VirtualProduction/CaptureManager/CaptureManagerApp/Source/IngestLiveLinkDevice/Private/BaseIngestLiveLinkDevice.cpp b/Engine/Plugins/VirtualProduction/CaptureManager/CaptureManagerApp/Source/IngestLiveLinkDevice/Private/BaseIngestLiveLinkDevice.cpp
--- a/Engine/Plugins/VirtualProduction/CaptureManager/CaptureManagerApp/Source/IngestLiveLinkDevice/Private/BaseIngestLiveLinkDevice.cpp
+++ b/Engine/Plugins/VirtualProduction/CaptureManager/CaptureManagerApp/Source/IngestLiveLinkDevice/Private/BaseIngestLiveLinkDevice.cpp
@@ -15,6 +15,18 @@
 
 DEFINE_LOG_CATEGORY(LogIngestLiveLinkDevice);
 
+namespace UE::CaptureManager::Private
+{
+
+TSharedRef<FCommunicationManager> GetCommunicationManager()
+{
+	FCaptureCommunicationManagerModule& CaptureCommunicationManagerModule =
+		FModuleManager::LoadModuleChecked<FCaptureCommunicationManagerModule>("CaptureCommunicationManager");
+	return CaptureCommunicationManagerModule.GetManager();
+}
+
+}
+
 class UBaseIngestLiveLinkDevice::FImpl
 {
 public:
diff --git a/Engine/Plugins/VirtualProduction/CaptureManager/CaptureManagerApp/Source/IngestLiveLinkDevice/Private/Utils/IngestLiveLinkDeviceUtils.cpp b/Engine/Plugins/VirtualProduction/CaptureManager/CaptureManagerApp/Source/IngestLiveLinkDevice/Private/Utils/IngestLiveLinkDeviceUtils.cpp
--- a/Engine/Plugins/VirtualProduction/CaptureManager/CaptureManagerApp/Source/IngestLiveLinkDevice/Private/Utils/IngestLiveLinkDeviceUtils.cpp
+++ b/Engine/Plugins/VirtualProduction/CaptureManager/CaptureManagerApp/Source/IngestLiveLinkDevice/Private/Utils/IngestLiveLinkDeviceUtils.cpp
@@ -1,16 +1,12 @@
 // Copyright Epic Games, Inc. All Rights Reserved.
 
-#include "IngestLiveLinkDeviceUtils.h"
+#include "Utils/IngestLiveLinkDeviceUtils.h"
 
-namespace UE::CaptureManager::Private
+namespace UE::CaptureManager
 {
 
-TSharedRef<FCommunicationManager> GetCommunicationManager()
+namespace Private
 {
-	FCaptureCommunicationManagerModule& CaptureCommunicationManagerModule =
-		FModuleManager::LoadModuleChecked<FCaptureCommunicationManagerModule>("CaptureCommunicationManager");
-	return CaptureCommunicationManagerModule.GetManager();
-}
 
 EMediaTexturePixelFormat ConvertPixelFormat(EIngestCapability_ImagePixelFormat InImagePixelFormat)
 {
@@ -59,6 +55,8 @@ EMediaOrientation ConvertImageRotation(EIngestCapability_ImageRotation InImageRo
 	}
 }
 
+}
+
 FString ErrorOriginToString(FTakeMetadataParserError::EOrigin InOrigin)
 {
 	switch (InOrigin)
diff --git a/Engine/Plugins/VirtualProduction/CaptureManager/CaptureManagerApp/Source/IngestLiveLinkDevice/Private/Utils/IngestLiveLinkDeviceUtils.h b/Engine/Plugins/VirtualProduction/CaptureManager/CaptureManagerApp/Source/IngestLiveLinkDevice/Private/Utils/IngestLiveLinkDeviceUtils.h
deleted file mode 100644
diff --git a/Engine/Plugins/VirtualProduction/CaptureManager/CaptureManagerApp/Source/IngestLiveLinkDevice/Public/Utils/IngestLiveLinkDeviceUtils.h b/Engine/Plugins/VirtualProduction/CaptureManager/CaptureManagerApp/Source/IngestLiveLinkDevice/Public/Utils/IngestLiveLinkDeviceUtils.h
new file mode 100644
--- /dev/null
+++ b/Engine/Plugins/VirtualProduction/CaptureManager/CaptureManagerApp/Source/IngestLiveLinkDevice/Public/Utils/IngestLiveLinkDeviceUtils.h
@@ -0,0 +1,19 @@
+// Copyright Epic Games, Inc. All Rights Reserved.
+
+#pragma once
+
+#include "CaptureManagerTakeMetadata.h"
+#include "Ingest/IngestCapability_Options.h"
+#include "MediaSample.h"
+
+namespace UE::CaptureManager
+{
+namespace Private
+{
+EMediaTexturePixelFormat ConvertPixelFormat(EIngestCapability_ImagePixelFormat InImagePixelFormat);
+EMediaOrientation ConvertImageRotation(EIngestCapability_ImageRotation InImageRotation);
+}
+
+INGESTLIVELINKDEVICE_API FString ErrorOriginToString(FTakeMetadataParserError::EOrigin InOrigin);
+
+}
diff --git a/Engine/Plugins/VirtualProduction/CaptureManager/CaptureManagerDevices/CaptureManagerDevices.uplugin b/Engine/Plugins/VirtualProduction/CaptureManager/CaptureManagerDevices/CaptureManagerDevices.uplugin
--- a/Engine/Plugins/VirtualProduction/CaptureManager/CaptureManagerDevices/CaptureManagerDevices.uplugin
+++ b/Engine/Plugins/VirtualProduction/CaptureManager/CaptureManagerDevices/CaptureManagerDevices.uplugin
@@ -21,6 +21,16 @@
 			"Name": "MonoVideoLiveLinkDevice",
 			"Type": "Editor",
 			"LoadingPhase": "Default"
+		},
+		{
+			"Name": "CPSLiveLinkDevice",
+			"Type": "Editor",
+			"LoadingPhase": "Default"
+		},
+		{
+			"Name": "FilesystemLiveLinkDevice",
+			"Type": "Editor",
+			"LoadingPhase": "Default"
 		}
 	],
 	"Plugins": [
diff --git a/Engine/Plugins/VirtualProduction/CaptureManager/CaptureManagerDevices/Source/CPSLiveLinkDevice/CPSLiveLinkDevice.Build.cs b/Engine/Plugins/VirtualProduction/CaptureManager/CaptureManagerDevices/Source/CPSLiveLinkDevice/CPSLiveLinkDevice.Build.cs
new file mode 100644
--- /dev/null
+++ b/Engine/Plugins/VirtualProduction/CaptureManager/CaptureManagerDevices/Source/CPSLiveLinkDevice/CPSLiveLinkDevice.Build.cs
@@ -0,0 +1,45 @@
+// Copyright Epic Games, Inc. All Rights Reserved.
+
+using UnrealBuildTool;
+using System.IO;
+
+public class CPSLiveLinkDevice : ModuleRules
+{
+	public CPSLiveLinkDevice(ReadOnlyTargetRules Target) : base(Target)
+	{
+		PCHUsage = PCHUsageMode.UseExplicitOrSharedPCHs;
+		ShortName = "CPSLLD";
+
+		PublicDependencyModuleNames.AddRange(new string[]
+		{
+			"CaptureUtils",
+			"Core"
+		});
+
+		PrivateDependencyModuleNames.AddRange(new string[]
+		{
+			"Engine",
+			"Json",
+			"CaptureManagerMediaRW",
+			"CaptureProtocolStack",
+			"IngestLiveLinkDevice",
+			"LiveLinkDevice",
+			"LiveLinkCapabilities",
+			"Media",
+			"LiveLinkHubCaptureMessaging",
+			"LiveLinkHubExportClient",
+			"LiveLinkHub",
+			"LiveLinkFaceMetadata",
+			"ToolWidgets",
+			"Slate",
+			"SlateCore",
+			"StereoCameraMetadata",
+			"CaptureManagerTakeMetadata",
+			"CaptureManagerSettings",
+			"CoreUObject",
+			"UnrealEd"
+		});
+
+		CppCompileWarningSettings.UnsafeTypeCastWarningLevel = WarningLevel.Error;
+	}
+}
diff --git a/Engine/Plugins/VirtualProduction/CaptureManager/CaptureManagerDevices/Source/CPSLiveLinkDevice/Private/Customizations/DeviceIpAddressCustomization.cpp b/Engine/Plugins/VirtualProduction/CaptureManager/CaptureManagerDevices/Source/CPSLiveLinkDevice/Private/Customizations/DeviceIpAddressCustomization.cpp
new file mode 100644
--- /dev/null
+++ b/Engine/Plugins/VirtualProduction/CaptureManager/CaptureManagerDevices/Source/CPSLiveLinkDevice/Private/Customizations/DeviceIpAddressCustomization.cpp
@@ -0,0 +1,136 @@
+// Copyright Epic Games, Inc. All Rights Reserved.
+
+#include "DeviceIpAddressCustomization.h"
+#include "Widgets/Input/SEditableTextBox.h"
+#include "Editor.h"
+
+#define LOCTEXT_NAMESPACE "DeviceIpAddressCustomization"
+
+FDeviceIpAddressCustomization::FDeviceIpAddressCustomization()
+	: IpAddressRange(TRange<int32>::Inclusive(0, 255))
+{
+}
+
+void FDeviceIpAddressCustomization::CustomizeHeader(TSharedRef<IPropertyHandle> PropertyHandle, FDetailWidgetRow& HeaderRow, IPropertyTypeCustomizationUtils& CustomizationUtils)
+{
+}
+
+void FDeviceIpAddressCustomization::CustomizeChildren(TSharedRef<IPropertyHandle> PropertyHandle, IDetailChildrenBuilder& ChildBuilder, IPropertyTypeCustomizationUtils& CustomizationUtils)
+{
+	DeviceIpAddressProperty = PropertyHandle;
+
+	ChildBuilder.AddProperty(PropertyHandle)
+		.CustomWidget()
+		.NameContent()
+		[
+			PropertyHandle->CreatePropertyNameWidget()
+		]
+		.ValueContent()
+		[
+			SNew(SEditableTextBox)
+			.Text(this, &FDeviceIpAddressCustomization::OnGetDeviceIpAddress)
+			.OnVerifyTextChanged(this, &FDeviceIpAddressCustomization::OnDeviceIpAddressVerify)
+			.OnTextCommitted(this, &FDeviceIpAddressCustomization::OnDeviceIpAddressCommited)
+			.SelectAllTextWhenFocused(true)
+			.RevertTextOnEscape(true)
+			.Font(IDetailLayoutBuilder::GetDetailFont())
+			.IsReadOnly(this, &FDeviceIpAddressCustomization::IsReadOnly)
+		];
+}
+
+TOptional<FText> FDeviceIpAddressCustomization::VerifyIpAddress(const FString& InIpAddress, TRange<int32> InIpAddressRange)
+{
+	if (InIpAddress.IsEmpty())
+	{
+		return LOCTEXT("DeviceIpAddress_Empty", "Device Ip Address must not be empty");
+	}
+
+	constexpr uint16 IPv4AddressMaxSize = 15;
+	if (InIpAddress.Len() > IPv4AddressMaxSize)
+	{
+		return LOCTEXT("DeviceIpAddress_InvalidSize", "Device Ip Address contains too many characters");
+	}
+
+	TArray<FString> Components;
+	InIpAddress.ParseIntoArray(Components, TEXT("."), false);
+
+	if (Components.Num() != 4)
+	{
+		return LOCTEXT("DeviceIpAddress_InvalidFormat", "Invalid format for Device Ip Address");
+	}
+
+	for (const FString& Component : Components)
+	{
+		if (!Component.IsNumeric())
+		{
+			return LOCTEXT("DeviceIpAddress_Numeric", "Device Ip Address must only contain numbers");
+		}
+
+		const int32 ComponentNum = FCString::Atoi(*Component);
+		if (!InIpAddressRange.Contains(ComponentNum))
+		{
+			return LOCTEXT("DeviceIpAddress_Numeric_NumbersInRange", "Device Ip Address must only contain numbers in range [0-255]");
+		}
+	}
+
+	return {};
+}
+
+FText FDeviceIpAddressCustomization::OnGetDeviceIpAddress() const
+{
+	TArray<void*> RawValue;
+	DeviceIpAddressProperty->AccessRawData(RawValue);
+
+	if (RawValue.Num())
+	{
+		const FDeviceIpAddress* DeviceAddress = static_cast<const FDeviceIpAddress*>(RawValue[0]);
+		return FText::FromString(DeviceAddress->IpAddressString);
+	}
+
+	return FText::GetEmpty();
+}
+
+bool FDeviceIpAddressCustomization::OnDeviceIpAddressVerify(const FText& InText, FText& OutErrorText)
+{
+	FString Value = InText.ToString();
+
+	TOptional<FText> MaybeError = VerifyIpAddress(InText.ToString(), IpAddressRange);
+
+	if (MaybeError.IsSet())
+	{
+		OutErrorText = MaybeError.GetValue();
+		return false;
+	}
+
+	return true;
+}
+
+void FDeviceIpAddressCustomization::OnDeviceIpAddressCommited(const FText& InText, ETextCommit::Type CommitInfo)
+{
+	TArray<void*> RawValue;
+	DeviceIpAddressProperty->AccessRawData(RawValue);
+
+	FString Value = InText.ToString();
+
+	if (RawValue.Num())
+	{
+		GEditor->BeginTransaction(FText::Format(LOCTEXT("DeviceIpAddress_SetProperty", "Edit {0}"), DeviceIpAddressProperty->GetPropertyDisplayName()));
+
+		DeviceIpAddressProperty->NotifyPreChange();
+
+		FDeviceIpAddress* DeviceIpAddress = static_cast<FDeviceIpAddress*>(RawValue[0]);
+		DeviceIpAddress->IpAddressString = MoveTemp(Value);
+
+		DeviceIpAddressProperty->NotifyPostChange(EPropertyChangeType::ValueSet);
+		DeviceIpAddressProperty->NotifyFinishedChangingProperties();
+
+		GEditor->EndTransaction();
+	}
+}
+
+bool FDeviceIpAddressCustomization::IsReadOnly() const
+{
+	return !DeviceIpAddressProperty->IsEditable();
+}
+
+#undef LOCTEXT_NAMESPACE
\ No newline at end of file
diff --git a/Engine/Plugins/VirtualProduction/CaptureManager/CaptureManagerDevices/Source/CPSLiveLinkDevice/Private/Customizations/DeviceIpAddressCustomization.h b/Engine/Plugins/VirtualProduction/CaptureManager/CaptureManagerDevices/Source/CPSLiveLinkDevice/Private/Customizations/DeviceIpAddressCustomization.h
new file mode 100644
--- /dev/null
+++ b/Engine/Plugins/VirtualProduction/CaptureManager/CaptureManagerDevices/Source/CPSLiveLinkDevice/Private/Customizations/DeviceIpAddressCustomization.h
@@ -0,0 +1,38 @@
+// Copyright Epic Games, Inc. All Rights Reserved.
+
+#pragma once
+
+#include "DetailLayoutBuilder.h"
+#include "DetailWidgetRow.h"
+#include "IDetailChildrenBuilder.h"
+
+#include "DeviceIpAddressCustomization.generated.h"
+
+USTRUCT(BlueprintType)
+struct FDeviceIpAddress
+{
+	GENERATED_BODY()
+
+	UPROPERTY(EditAnywhere, BlueprintReadWrite, Category = DeviceAddress, DisplayName="Ip Address")
+	FString IpAddressString;
+};
+
+class FDeviceIpAddressCustomization : public IPropertyTypeCustomization
+{
+public:
+	FDeviceIpAddressCustomization();
+
+	virtual void CustomizeHeader(TSharedRef<IPropertyHandle> PropertyHandle, FDetailWidgetRow& HeaderRow, IPropertyTypeCustomizationUtils& CustomizationUtils) override;
+	virtual void CustomizeChildren(TSharedRef<IPropertyHandle> PropertyHandle, IDetailChildrenBuilder& ChildBuilder, IPropertyTypeCustomizationUtils& CustomizationUtils) override;
+
+	static TOptional<FText> VerifyIpAddress(const FString& InIpAddress, TRange<int32> InIpAddressRange);
+
+private:
+	FText OnGetDeviceIpAddress() const;
+	bool OnDeviceIpAddressVerify(const FText& InText, FText& OutErrorText);
+	void OnDeviceIpAddressCommited(const FText& InText, ETextCommit::Type CommitInfo);
+	bool IsReadOnly() const;
+
+	TSharedPtr<IPropertyHandle> DeviceIpAddressProperty;
+	TRange<int32> IpAddressRange;
+};
diff --git a/Engine/Plugins/VirtualProduction/CaptureManager/CaptureManagerDevices/Source/CPSLiveLinkDevice/Private/Customizations/ToggleConnectActionCustomization.cpp b/Engine/Plugins/VirtualProduction/CaptureManager/CaptureManagerDevices/Source/CPSLiveLinkDevice/Private/Customizations/ToggleConnectActionCustomization.cpp
new file mode 100644
--- /dev/null
+++ b/Engine/Plugins/VirtualProduction/CaptureManager/CaptureManagerDevices/Source/CPSLiveLinkDevice/Private/Customizations/ToggleConnectActionCustomization.cpp
@@ -0,0 +1,210 @@
+// Copyright Epic Games, Inc. All Rights Reserved.
+
+#include "Customizations/ToggleConnectActionCustomization.h"
+
+#include "Customizations/DeviceIpAddressCustomization.h"
+
+#include "DetailWidgetRow.h"
+#include "SSimpleButton.h"
+#include "Widgets/Images/SImage.h"
+
+#include "LiveLinkDeviceSubsystem.h"
+#include "Engine/Engine.h"
+
+#include "LiveLinkFaceDevice.h"
+
+#include "LiveLinkDeviceCapability_Connection.h"
+
+#define LOCTEXT_NAMESPACE "ToggleConnectActionCustomization"
+
+void FToggleConnectActionCustomization::CustomizeHeader(TSharedRef<IPropertyHandle> InPropertyHandle,
+														FDetailWidgetRow& InHeaderRow,
+														IPropertyTypeCustomizationUtils& InCustomizationUtils)
+{
+	PropertyHandle = InPropertyHandle;
+
+	Device = GetDevice();
+
+	if (ensure(Device))
+	{
+		UConnectionDelegate* ConnectionDelegate = ILiveLinkDeviceCapability_Connection::Execute_GetConnectionDelegate(Device);
+		ConnectionDelegate->ConnectionChanged.AddSP(this, &FToggleConnectActionCustomization::OnConnectionStatusChanged);
+	}
+
+	PropertyHandle->MarkResetToDefaultCustomized();
+	PropertyHandle->MarkHiddenByCustomization();
+
+	InHeaderRow
+		.NameContent()
+		[
+			PropertyHandle->CreatePropertyNameWidget(LOCTEXT("ConnectToggleActionName", "Connect/Disconnect"))
+		]
+		.ValueContent()
+		[
+			SNew(SButton)
+			.HAlign(HAlign_Center)
+			.VAlign(VAlign_Center)
+			.Text(this, &FToggleConnectActionCustomization::GetButtonName)
+			.OnClicked(this, &FToggleConnectActionCustomization::OnConnectButtonToggled)
+			.IsEnabled(this, &FToggleConnectActionCustomization::IsConnectButtonEnabled)
+		]
+		.ResetToDefaultContent()
+		[
+			SNew(SButton)
+			.IsFocusable(false)
+			.ButtonStyle(FAppStyle::Get(), "SimpleButton")
+			.ContentPadding(0.0f)
+			.Visibility(this, &FToggleConnectActionCustomization::GetStopActionVisibility)
+			.OnClicked(this, &FToggleConnectActionCustomization::OnStopActionClicked)
+			.ToolTipText(this, &FToggleConnectActionCustomization::GetStopActionTooltip)
+			.Content()
+			[
+				SNew(SImage)
+				.Image(FAppStyle::GetBrush("GenericStop"))
+				.ColorAndOpacity(FSlateColor::UseForeground())
+			]
+		];
+}
+
+void FToggleConnectActionCustomization::CustomizeChildren(TSharedRef<IPropertyHandle> InPropertyHandle,
+														  IDetailChildrenBuilder& InChildBuilder,
+														  IPropertyTypeCustomizationUtils& InCustomizationUtils)
+{
+}
+
+void FToggleConnectActionCustomization::OnConnectionStatusChanged(ELiveLinkDeviceConnectionStatus InStatus)
+{
+	if (InStatus == ELiveLinkDeviceConnectionStatus::Disconnected)
+	{
+		bIsDisconnecting = false;
+	}
+}
+
+FText FToggleConnectActionCustomization::GetButtonName() const
+{
+	ELiveLinkDeviceConnectionStatus Status = ILiveLinkDeviceCapability_Connection::Execute_GetConnectionStatus(Device);
+
+	if (Status == ELiveLinkDeviceConnectionStatus::Connecting)
+	{
+		return LOCTEXT("ConnectingState", "Connecting");
+	}
+	else if (bIsDisconnecting || Status == ELiveLinkDeviceConnectionStatus::Disconnecting)
+	{
+		return LOCTEXT("DisconnectingState", "Disconnecting");
+	}
+	else if (Status == ELiveLinkDeviceConnectionStatus::Connected)
+	{
+		return LOCTEXT("ConnectedState", "Disconnect");
+	}
+	else if (Status == ELiveLinkDeviceConnectionStatus::Disconnected)
+	{
+		return LOCTEXT("DisconnectedState", "Connect");
+	}
+	else
+	{
+		return FText();
+	}
+}
+
+FReply FToggleConnectActionCustomization::OnConnectButtonToggled()
+{
+	ELiveLinkDeviceConnectionStatus Status = ILiveLinkDeviceCapability_Connection::Execute_GetConnectionStatus(Device);
+	if (Status == ELiveLinkDeviceConnectionStatus::Disconnected)
+	{
+		ILiveLinkDeviceCapability_Connection::Execute_Connect(Device);
+	}
+	else if (Status == ELiveLinkDeviceConnectionStatus::Connected)
+	{
+		bIsDisconnecting = true;
+		ILiveLinkDeviceCapability_Connection::Execute_Disconnect(Device);
+	}
+	else
+	{
+		return FReply::Unhandled();
+	}
+	
+	return FReply::Handled();
+}
+
+bool FToggleConnectActionCustomization::IsConnectButtonEnabled() const
+{
+	ELiveLinkDeviceConnectionStatus Status = ILiveLinkDeviceCapability_Connection::Execute_GetConnectionStatus(Device);
+
+	if (Status == ELiveLinkDeviceConnectionStatus::Connecting ||
+		bIsDisconnecting || Status == ELiveLinkDeviceConnectionStatus::Disconnecting)
+	{
+		return false;
+	}
+
+	if (CastChecked<ULiveLinkFaceDevice>(Device)->GetSettings()->IpAddress.IpAddressString.IsEmpty())
+	{
+		return false;
+	}
+
+	return true;
+}
+
+EVisibility FToggleConnectActionCustomization::GetStopActionVisibility() const
+{
+	ELiveLinkDeviceConnectionStatus Status = ILiveLinkDeviceCapability_Connection::Execute_GetConnectionStatus(Device);
+
+	if (Status == ELiveLinkDeviceConnectionStatus::Connecting)
+	{
+		return EVisibility::Visible;
+	}
+
+	return EVisibility::Hidden;
+}
+
+FReply FToggleConnectActionCustomization::OnStopActionClicked()
+{
+	ELiveLinkDeviceConnectionStatus Status = ILiveLinkDeviceCapability_Connection::Execute_GetConnectionStatus(Device);
+
+	if (Status == ELiveLinkDeviceConnectionStatus::Connecting)
+	{
+		bIsDisconnecting = true;
+
+		ILiveLinkDeviceCapability_Connection::Execute_Disconnect(Device);
+
+		bIsDisconnecting = false;
+
+		return FReply::Handled();
+	}
+
+	return FReply::Unhandled();
+}
+
+FText FToggleConnectActionCustomization::GetStopActionTooltip() const
+{
+	return LOCTEXT("StopConnectActionTooltip", "Stop the connect procedure");
+}
+
+TObjectPtr<ULiveLinkDevice> FToggleConnectActionCustomization::GetDevice() const
+{
+	TArray<void*> RawValue;
+	PropertyHandle->AccessRawData(RawValue);
+
+	if (RawValue.IsEmpty())
+	{
+		return nullptr;
+	}
+
+	const FToggleConnectAction* ConnectAction = static_cast<const FToggleConnectAction*>(RawValue[0]);
+
+	ULiveLinkDeviceSubsystem* Subsystem = GEngine->GetEngineSubsystem<ULiveLinkDeviceSubsystem>();
+	const TObjectPtr<ULiveLinkDevice>* MaybeDevice = Subsystem->GetDeviceMap().Find(ConnectAction->DeviceGuid);
+
+	if (!MaybeDevice)
+	{
+		return nullptr;
+	}
+
+	if (!(*MaybeDevice)->Implements<ULiveLinkDeviceCapability_Connection>())
+	{
+		return nullptr;
+	}
+
+	return *MaybeDevice;
+}
+
+#undef LOCTEXT_NAMESPACE
\ No newline at end of file
diff --git a/Engine/Plugins/VirtualProduction/CaptureManager/CaptureManagerDevices/Source/CPSLiveLinkDevice/Private/Customizations/ToggleConnectActionCustomization.h b/Engine/Plugins/VirtualProduction/CaptureManager/CaptureManagerDevices/Source/CPSLiveLinkDevice/Private/Customizations/ToggleConnectActionCustomization.h
new file mode 100644
--- /dev/null
+++ b/Engine/Plugins/VirtualProduction/CaptureManager/CaptureManagerDevices/Source/CPSLiveLinkDevice/Private/Customizations/ToggleConnectActionCustomization.h
@@ -0,0 +1,49 @@
+// Copyright Epic Games, Inc. All Rights Reserved.
+
+#pragma once
+
+#include "IPropertyTypeCustomization.h"
+#include "Input/Reply.h"
+
+#include "LiveLinkDevice.h"
+#include "LiveLinkDeviceCapability_Connection.h"
+
+#include "ToggleConnectActionCustomization.generated.h"
+
+USTRUCT()
+struct FToggleConnectAction
+{
+	GENERATED_BODY()
+
+	UPROPERTY()
+	FGuid DeviceGuid;
+};
+
+class FToggleConnectActionCustomization : public IPropertyTypeCustomization
+{
+public:
+
+	virtual void CustomizeHeader(TSharedRef<IPropertyHandle> InPropertyHandle, 
+								 FDetailWidgetRow& InHeaderRow, 
+								 IPropertyTypeCustomizationUtils& InCustomizationUtils) override;
+	virtual void CustomizeChildren(TSharedRef<IPropertyHandle> InPropertyHandle, 
+								   IDetailChildrenBuilder& InChildBuilder, 
+								   IPropertyTypeCustomizationUtils& InCustomizationUtils) override;
+private:
+
+	void OnConnectionStatusChanged(ELiveLinkDeviceConnectionStatus InStatus);
+
+	FText GetButtonName() const;
+	FReply OnConnectButtonToggled();
+	bool IsConnectButtonEnabled() const;
+
+	EVisibility GetStopActionVisibility() const;
+	FReply OnStopActionClicked();
+	FText GetStopActionTooltip() const;
+
+	TObjectPtr<ULiveLinkDevice> GetDevice() const;
+
+	TSharedPtr<IPropertyHandle> PropertyHandle;
+	TObjectPtr<ULiveLinkDevice> Device;
+	std::atomic_bool bIsDisconnecting = false;
+};
\ No newline at end of file
diff --git a/Engine/Plugins/VirtualProduction/CaptureManager/CaptureManagerDevices/Source/CPSLiveLinkDevice/Private/LiveLinkFaceDevice.cpp b/Engine/Plugins/VirtualProduction/CaptureManager/CaptureManagerDevices/Source/CPSLiveLinkDevice/Private/LiveLinkFaceDevice.cpp
new file mode 100644
--- /dev/null
+++ b/Engine/Plugins/VirtualProduction/CaptureManager/CaptureManagerDevices/Source/CPSLiveLinkDevice/Private/LiveLinkFaceDevice.cpp
@@ -0,0 +1,489 @@
+// Copyright Epic Games, Inc. All Rights Reserved.
+
+#include "LiveLinkFaceDevice.h"
+
+#include "Settings/CaptureManagerSettings.h"
+
+#include "ILiveLinkRecordingSessionInfo.h"
+
+#include "Control/Messages/Constants.h"
+#include "Control/Messages/ControlUpdate.h"
+
+#include "Protocol/CPSDataStream.h"
+#include "Protocol/CPSFileStream.h"
+
+#include "LiveLinkFaceMetadata.h"
+#include "StereoCameraMetadataParseUtils.h"
+
+#include "HAL/FileManager.h"
+
+const ULiveLinkFaceDeviceSettings* ULiveLinkFaceDevice::GetSettings() const
+{
+	return CastChecked<ULiveLinkFaceDeviceSettings>(Settings);
+}
+
+TSubclassOf<ULiveLinkDeviceSettings> ULiveLinkFaceDevice::StaticGetSettingsClass() const
+{
+	return ULiveLinkFaceDeviceSettings::StaticClass();
+}
+
+FText ULiveLinkFaceDevice::GetDisplayName() const
+{
+	return FText::FromString(GetSettings()->DisplayName);
+}
+
+EDeviceHealth ULiveLinkFaceDevice::GetDeviceHealth() const
+{
+	return EDeviceHealth::Nominal;
+}
+
+FText ULiveLinkFaceDevice::GetHealthText() const
+{
+	return FText::FromString("Example Health");
+}
+
+void ULiveLinkFaceDevice::OnDeviceAdded(const FGuid InDeviceGuid, ULiveLinkDeviceSettings* InSettings)
+{
+	CastChecked<ULiveLinkFaceDeviceSettings>(InSettings)->ConnectAction.DeviceGuid = InDeviceGuid;
+
+	Super::OnDeviceAdded(InDeviceGuid, InSettings);
+}
+
+void ULiveLinkFaceDevice::OnDeviceRemoved()
+{
+	ILiveLinkDeviceCapability_Connection::Execute_Disconnect(this);
+
+	Super::OnDeviceRemoved();
+}
+
+FString ULiveLinkFaceDevice::GetFullTakePath(UE::CaptureManager::FTakeId InTakeId) const
+{
+	FScopeLock Lock(&DownloadedTakesMutex);
+	if (const FString* DownloadedTake = DownloadedTakes.Find(InTakeId))
+	{
+		return *DownloadedTake;
+	}
+	
+	return FString();
+}
+
+void ULiveLinkFaceDevice::UpdateTakeList_Implementation()
+{
+	using namespace UE::CaptureManager;
+
+	RemoveAllTakes();
+
+	TProtocolResult<TArray<FGetTakeMetadataResponse::FTakeObject>> TakesResult = Device->FetchTakeList();
+
+	if (TakesResult.HasError())
+	{
+		return;
+	}
+
+	TMap<FString, FTakeId> NameToIdMap;
+
+	for (const FGetTakeMetadataResponse::FTakeObject& Take : TakesResult.StealValue())
+	{
+		FTakeId TakeId = AddTake(ParseTakeMetadata(Take));
+		Device->AddTakeMetadata(TakeId, Take);
+		NameToIdMap.Add(Take.Name, TakeId);
+	}
+
+	FetchThumbnails(MoveTemp(NameToIdMap));
+
+	ExecuteUpdateTakeListCallback(Execute_GetTakeIdentifiers(this));
+}
+
+void ULiveLinkFaceDevice::RunIngestTake(int32 InTakeId, const UIngestCapability_Options* InIngestOptions)
+{
+	using namespace UE::CaptureManager;
+
+	const FGetTakeMetadataResponse::FTakeObject& Take = Device->GetTake(InTakeId);
+
+	static constexpr uint32 NumberOfTasks = 3; // Download, Convert, Upload
+
+	TSharedPtr<FTaskProgress> TaskProgress = MakeShared<FTaskProgress>
+		(NumberOfTasks, FTaskProgress::FProgressReporter::CreateLambda([this, TakeId = InTakeId](double InProgress)
+	{
+		ExecuteProcessProgressReporter(TakeId, InProgress);
+	}));
+
+	uint64 TotalSize = 0;
+	for (const FGetTakeMetadataResponse::FFileObject& File : Take.Files)
+	{
+		TotalSize += File.Length;
+	}
+
+	FString DownloadedStorage = InIngestOptions->DownloadDirectory;
+
+	FTaskProgress::FTask Task = TaskProgress->StartTask();
+
+	TStrongObjectPtr<const UIngestCapability_Options> IngestOptions(InIngestOptions);
+
+	TUniquePtr<FCPSFileStream> CPSStream = MakeUnique<FCPSFileStream>(MoveTemp(DownloadedStorage), TotalSize);
+	CPSStream->SetExportFinished(FCPSFileStream::FExportFinished::CreateUObject(this, &ULiveLinkFaceDevice::OnExportFinished, Take.Name, InTakeId, TaskProgress, IngestOptions));
+	CPSStream->SetProgressHandler(FCPSFileStream::FReportProgress::CreateUObject(this, &ULiveLinkFaceDevice::OnExportProgressReport, Task));
+
+	Device->StartExport(InTakeId, MoveTemp(CPSStream));
+}
+
+void ULiveLinkFaceDevice::CancelIngest(int32 InTakeId)
+{
+	Device->CancelExport(InTakeId);
+
+	Super::CancelIngest(InTakeId);
+}
+
+ELiveLinkDeviceConnectionStatus ULiveLinkFaceDevice::GetConnectionStatus_Implementation() const
+{
+	if (!Device)
+	{
+		return ELiveLinkDeviceConnectionStatus::Disconnected;
+	}
+
+	if (bIsConnecting)
+	{
+		return ELiveLinkDeviceConnectionStatus::Connecting;
+	}
+	else if (Device->IsConnected())
+	{
+		return ELiveLinkDeviceConnectionStatus::Connected;
+	}
+
+	return ELiveLinkDeviceConnectionStatus::Disconnected;
+}
+
+FString ULiveLinkFaceDevice::GetHardwareId_Implementation() const
+{
+	return GetSettings()->IpAddress.IpAddressString;
+}
+
+bool ULiveLinkFaceDevice::SetHardwareId_Implementation(const FString& HardwareID)
+{
+	return false;
+}
+
+bool ULiveLinkFaceDevice::Connect_Implementation()
+{
+	const ULiveLinkFaceDeviceSettings* DeviceSettings = GetSettings();
+
+	using namespace UE::CaptureManager;
+
+	Device = FCPSDevice::MakeCPSDevice(DeviceSettings->IpAddress.IpAddressString, static_cast<uint16>(DeviceSettings->Port));
+	
+	Device->SubscribeToEvent(FConnectionStateChangedEvent::Name, FCaptureEventHandler(
+		FCaptureEventHandler::Type::CreateUObject(this, &ULiveLinkFaceDevice::HandleConnectionChanged),
+		EDelegateExecutionThread::AnyThread));
+
+	Device->SubscribeToEvent(FCPSStateEvent::Name, FCaptureEventHandler(
+		FCaptureEventHandler::Type::CreateUObject(this, &ULiveLinkFaceDevice::HandleCPSStateUpdate),
+		EDelegateExecutionThread::AnyThread));
+
+	Device->SubscribeToEvent(FCPSEvent::Name, FCaptureEventHandler(
+		FCaptureEventHandler::Type::CreateUObject(this, &ULiveLinkFaceDevice::HandleCPSEvent),
+		EDelegateExecutionThread::AnyThread));
+
+	Device->InitiateConnect();
+
+	return true;
+}
+
+bool ULiveLinkFaceDevice::Disconnect_Implementation()
+{
+	AsyncTask(ENamedThreads::AnyThread, [this]()
+	{
+		if (Device)
+		{
+			Device->Stop();
+			Device->UnsubscribeAll();
+			bIsConnecting = false;
+			Device = nullptr;
+		}
+	});
+	
+	return true;
+}
+
+bool ULiveLinkFaceDevice::StartRecording_Implementation()
+{
+	if (!Device)
+	{
+		return false;
+	}
+
+	using namespace UE::CaptureManager;
+
+	ILiveLinkRecordingSessionInfo& SessionInfo = ILiveLinkRecordingSessionInfo::Get();
+	if (SessionInfo.GetSlateName().IsEmpty() || SessionInfo.GetTakeNumber() == -1)
+	{
+		return false;
+	}
+
+	uint16 TakeNumber = static_cast<uint16>(SessionInfo.GetTakeNumber());
+	TProtocolResult<void> Result = Device->StartRecording(SessionInfo.GetSlateName(), TakeNumber);
+
+	return Result.HasValue();
+}
+
+bool ULiveLinkFaceDevice::StopRecording_Implementation()
+{
+	if (!Device)
+	{
+		return false;
+	}
+
+	using namespace UE::CaptureManager;
+
+	TProtocolResult<void> Result = Device->StopRecording();
+
+	return Result.HasValue();
+}
+
+bool ULiveLinkFaceDevice::IsRecording_Implementation() const
+{
+	return bIsRecording;
+}
+
+void ULiveLinkFaceDevice::HandleConnectionChanged(TSharedPtr<const UE::CaptureManager::FCaptureEvent> InEvent)
+{
+	using namespace UE::CaptureManager;
+
+	TSharedPtr<const FConnectionStateChangedEvent> Event = StaticCastSharedPtr<const FConnectionStateChangedEvent>(InEvent);
+
+	if (Event->ConnectionState == FConnectionStateChangedEvent::EState::Connecting)
+	{
+		bIsConnecting = true;
+	}
+	else if (Event->ConnectionState == FConnectionStateChangedEvent::EState::Connected)
+	{
+		bIsConnecting = false;
+	}
+
+	auto GetConnectionState = [](FConnectionStateChangedEvent::EState InConnectionState) -> ELiveLinkDeviceConnectionStatus
+	{
+		switch (InConnectionState)
+		{
+			case FConnectionStateChangedEvent::EState::Connecting:
+				return ELiveLinkDeviceConnectionStatus::Connecting;
+			case FConnectionStateChangedEvent::EState::Connected:
+				return ELiveLinkDeviceConnectionStatus::Connected;
+			case FConnectionStateChangedEvent::EState::Disconnected:
+			case FConnectionStateChangedEvent::EState::Unknown:
+			default:
+				return ELiveLinkDeviceConnectionStatus::Disconnected;
+		}
+	};
+
+	ELiveLinkDeviceConnectionStatus ConnectionStatus = GetConnectionState(Event->ConnectionState);
+
+	SetConnectionStatus(ConnectionStatus);
+}
+
+void ULiveLinkFaceDevice::HandleCPSStateUpdate(TSharedPtr<const UE::CaptureManager::FCaptureEvent> InEvent)
+{
+	using namespace UE::CaptureManager;
+
+	TSharedPtr<const FCPSStateEvent> Event = StaticCastSharedPtr<const FCPSStateEvent>(InEvent);
+
+	bIsRecording = Event->GetStateResponse.IsRecording();
+}
+
+void ULiveLinkFaceDevice::HandleCPSEvent(TSharedPtr<const UE::CaptureManager::FCaptureEvent> InEvent)
+{
+	using namespace UE::CaptureManager;
+
+	TSharedPtr<const FCPSEvent> Event = StaticCastSharedPtr<const FCPSEvent>(InEvent);
+
+	if (Event->UpdateMessage->GetAddressPath() == CPS::AddressPaths::GRecordingStatus)
+	{
+		TSharedPtr<FRecordingStatusUpdate> Update = StaticCastSharedPtr<FRecordingStatusUpdate>(Event->UpdateMessage);
+		bIsRecording = Update->IsRecording();
+	}
+	else if (Event->UpdateMessage->GetAddressPath() == CPS::AddressPaths::GTakeAdded)
+	{
+		TSharedPtr<const FTakeAddedUpdate> TakeAddedUpdate = StaticCastSharedPtr<const FTakeAddedUpdate>(Event->UpdateMessage);
+
+		TProtocolResult<FGetTakeMetadataResponse::FTakeObject> FetchTakeResult = Device->FetchTake(TakeAddedUpdate->GetTakeName());
+
+		if (FetchTakeResult.IsValid())
+		{
+			FGetTakeMetadataResponse::FTakeObject Take = FetchTakeResult.StealValue();
+			
+			FString TakeName = Take.Name;
+
+			FTakeId TakeId = AddTake(ParseTakeMetadata(Take));
+			Device->AddTakeMetadata(TakeId, MoveTemp(Take));
+
+			PublishEvent<FIngestCapability_TakeAddedEvent>(TakeId);
+
+			FetchThumbnail(TakeName, TakeId);
+		}
+	}
+	else if (Event->UpdateMessage->GetAddressPath() == CPS::AddressPaths::GTakeRemoved)
+	{
+		TSharedPtr<const FTakeRemovedUpdate> TakeRemoveUpdate = StaticCastSharedPtr<const FTakeRemovedUpdate>(Event->UpdateMessage);
+
+		FTakeId TakeId = Device->GetTakeId(TakeRemoveUpdate->GetTakeName());
+
+		if (TakeId != INDEX_NONE)
+		{
+			CancelIngest(TakeId);
+
+			RemoveTake(TakeId);
+			Device->RemoveTakeMetadata(TakeId);
+			
+			PublishEvent<FIngestCapability_TakeRemovedEvent>(TakeId);
+		}
+	}
+}
+
+FTakeMetadata ULiveLinkFaceDevice::ParseTakeMetadata(const UE::CaptureManager::FGetTakeMetadataResponse::FTakeObject& InTake)
+{
+	FTakeMetadata TakeMetadata;
+
+	TakeMetadata.Slate = InTake.Slate;
+	TakeMetadata.TakeNumber = InTake.TakeNumber;
+	FDateTime DateTime = TakeMetadata.DateTime.Get(FDateTime());
+	FDateTime::ParseIso8601(*InTake.DateTime, DateTime);
+	TakeMetadata.DateTime = DateTime;
+
+	FTakeMetadata::FVideo Video;
+	Video.FrameRate = InTake.Video.FrameRate;
+	Video.FramesCount = static_cast<uint32>(InTake.Video.Frames);
+	Video.Format = TEXT("mov");
+	Video.FrameHeight = InTake.Video.Height;
+	Video.FrameWidth = InTake.Video.Width;
+
+	TakeMetadata.Video.Add(MoveTemp(Video));
+
+	return TakeMetadata;
+}
+
+void ULiveLinkFaceDevice::OnExportProgressReport(float InProgress,
+													   UE::CaptureManager::FTaskProgress::FTask InTask)
+{
+	InTask.Update(static_cast<double>(InProgress));
+}
+
+void ULiveLinkFaceDevice::OnExportFinished(UE::CaptureManager::TProtocolResult<void> InResult,
+												 FString InTakeName,
+												 UE::CaptureManager::FTakeId InTakeId,
+												 TSharedPtr<UE::CaptureManager::FTaskProgress> InTaskProgress,
+												 TStrongObjectPtr<const UIngestCapability_Options> InIngestOptions)
+{
+	using namespace UE::CaptureManager;
+
+	const FString DownloadedStorage = InIngestOptions->DownloadDirectory;
+	const FString DownloadedTake = FPaths::Combine(DownloadedStorage, InTakeName);
+
+	if (InResult.HasValue())
+	{
+		// Blocking function
+		TOptional<FTakeMetadata> TakeMetadata = ParseTake(DownloadedStorage, InTakeName);
+		if (TakeMetadata.IsSet())
+		{
+			FScopeLock Lock(&DownloadedTakesMutex);
+			DownloadedTakes.Add(InTakeId, DownloadedTake);
+			Lock.Unlock();
+
+			UpdateTake(InTakeId, MoveTemp(TakeMetadata.GetValue()));
+			Super::IngestTake(InTakeId, InIngestOptions.Get(), InTaskProgress);
+		}
+		else
+		{
+			TValueOrError<void, FIngestCapability_Error> Result = MakeError(FIngestCapability_Error::DownloaderError, TEXT("Failed to parse the take metadata"));
+			ExecuteProcessFinishedReporter(InTakeId, MoveTemp(Result));
+		}
+	}
+	else
+	{
+		TValueOrError<void, FIngestCapability_Error> Result = MakeError(FIngestCapability_Error::DownloaderError, InResult.GetError().GetMessage());
+
+		ExecuteProcessFinishedReporter(InTakeId, MoveTemp(Result));
+	}
+
+	IFileManager::Get().DeleteDirectory(*DownloadedTake, false, true);
+}
+
+TOptional<FTakeMetadata> ULiveLinkFaceDevice::ParseTake(const FString& InTakeDirectory, const FString& InTakeName)
+{
+	using namespace UE::CaptureManager;
+
+	const FString TakeDirectory = InTakeDirectory;
+	const FString TakePath = FPaths::Combine(TakeDirectory, InTakeName);
+
+	TArray<FString> TakeFiles;
+	IFileManager::Get().FindFiles(TakeFiles, *TakePath, *FTakeMetadata::FileExtension);
+	if (!TakeFiles.IsEmpty())
+	{
+		check(TakeFiles.Num() == 1);
+
+		FTakeMetadataParser TakeMetadataParser;
+		TValueOrError<FTakeMetadata, FTakeMetadataParserError> TakeMetadataResult = TakeMetadataParser.Parse(TakePath / TakeFiles[0]);
+		if (!TakeMetadataResult.HasError())
+		{
+			return TakeMetadataResult.StealValue();
+		}
+	}
+
+	TArray<FText> ValidationFailures;
+	TOptional<FTakeMetadata> ParseOldMetadataResult = LiveLinkMetadata::ParseOldLiveLinkTakeMetadata(TakePath, ValidationFailures);
+
+	if (ParseOldMetadataResult.IsSet())
+	{
+		return ParseOldMetadataResult.GetValue();
+	}
+
+	ParseOldMetadataResult = StereoCameraMetadata::ParseOldStereoCameraMetadata(TakePath, ValidationFailures);
+
+	if (ParseOldMetadataResult.IsSet())
+	{
+		return ParseOldMetadataResult.GetValue();
+	}
+
+	return {};
+}
+
+void ULiveLinkFaceDevice::FetchThumbnail(const FString& InTakeName, UE::CaptureManager::FTakeId InTakeId)
+{
+	TMap<FString, UE::CaptureManager::FTakeId> NameToIdMap;
+	NameToIdMap.Add(InTakeName, InTakeId);
+
+	FetchThumbnails(MoveTemp(NameToIdMap));
+}
+
+void ULiveLinkFaceDevice::FetchThumbnails(TMap<FString, UE::CaptureManager::FTakeId> InNameToIdMap)
+{
+	using namespace UE::CaptureManager;
+
+	FCPSDataStream::FFileExportFinished OnFileExportFinishedCallback = FCPSDataStream::FFileExportFinished::CreateLambda(
+		[this, NameToIdMap = MoveTemp(InNameToIdMap)](const FString& InTakeName, TProtocolResult<FCPSDataStream::FData> InData)
+		{
+			if (InData.HasError())
+			{
+				return;
+			}
+
+			FCPSDataStream::FData Data = InData.StealValue();
+
+			if (const FTakeId* TakeId = NameToIdMap.Find(InTakeName))
+			{
+				TOptional<FTakeMetadata> TakeMetadataOpt = GetTakeMetadata(*TakeId);
+				if (!TakeMetadataOpt.IsSet())
+				{
+					return;
+				}
+
+				FTakeMetadata TakeMetadata = TakeMetadataOpt.GetValue();
+				TakeMetadata.Thumbnail = FTakeThumbnailData(MoveTemp(Data));
+
+				UpdateTake(*TakeId, MoveTemp(TakeMetadata));
+
+				PublishEvent<FIngestCapability_TakeUpdatedEvent>(*TakeId);
+			}
+		});
+
+	TUniquePtr<FCPSDataStream> DataStream = MakeUnique<FCPSDataStream>(MoveTemp(OnFileExportFinishedCallback));
+
+	Device->FetchThumbnails(MoveTemp(DataStream));
+}
diff --git a/Engine/Plugins/VirtualProduction/CaptureManager/CaptureManagerDevices/Source/CPSLiveLinkDevice/Private/LiveLinkFaceDevice.h b/Engine/Plugins/VirtualProduction/CaptureManager/CaptureManagerDevices/Source/CPSLiveLinkDevice/Private/LiveLinkFaceDevice.h
new file mode 100644
--- /dev/null
+++ b/Engine/Plugins/VirtualProduction/CaptureManager/CaptureManagerDevices/Source/CPSLiveLinkDevice/Private/LiveLinkFaceDevice.h
@@ -0,0 +1,112 @@
+// Copyright Epic Games, Inc. All Rights Reserved.
+
+#pragma once
+
+#include "LiveLinkDevice.h"
+
+#include "Ingest/LiveLinkDeviceCapability_Ingest.h"
+#include "LiveLinkDeviceCapability_Connection.h"
+#include "LiveLinkDeviceCapability_Recording.h"
+#include "Misc/Optional.h"
+
+#include "Async/TaskProgress.h"
+
+#include "Protocol/CPSDevice.h"
+#include "Customizations/ToggleConnectActionCustomization.h"
+#include "Customizations/DeviceIpAddressCustomization.h"
+
+#include "BaseIngestLiveLinkDevice.h"
+
+#include "LiveLinkFaceDevice.generated.h"
+
+
+UCLASS(BlueprintType)
+class ULiveLinkFaceDeviceSettings : public ULiveLinkDeviceSettings
+{
+	GENERATED_BODY()
+
+public:
+	UPROPERTY(EditAnywhere, BlueprintReadWrite, Category = "Live Link Face Device")
+	FString DisplayName = TEXT("Live Link Face Device");
+
+	UPROPERTY(EditAnywhere, BlueprintReadWrite, Category = "Live Link Face Device")
+	FDeviceIpAddress IpAddress;
+
+	UPROPERTY(EditAnywhere, BlueprintReadWrite, Category = "Live Link Face Device")
+	int32 Port = 14785;
+
+	UPROPERTY(VisibleAnywhere, Category = "Live Link Face Device")
+	FToggleConnectAction ConnectAction;
+};
+
+
+UCLASS(BlueprintType)
+class ULiveLinkFaceDevice final : public UBaseIngestLiveLinkDevice
+	, public ILiveLinkDeviceCapability_Connection
+	, public ILiveLinkDeviceCapability_Recording
+{
+	GENERATED_BODY()
+
+public:
+	UFUNCTION(BlueprintCallable, Category = "LiveLinkFaceDeviceSettings")
+	const ULiveLinkFaceDeviceSettings* GetSettings() const;
+
+	//~ Begin ULiveLinkDevice interface
+	virtual TSubclassOf<ULiveLinkDeviceSettings> StaticGetSettingsClass() const override;
+	virtual FText GetDisplayName() const override;
+	virtual EDeviceHealth GetDeviceHealth() const override;
+	virtual FText GetHealthText() const override;
+
+	virtual void OnDeviceAdded(const FGuid InDeviceGuid, ULiveLinkDeviceSettings* InSettings) override;
+	virtual void OnDeviceRemoved() override;
+	//~ End ULiveLinkDevice interface
+
+private:
+	virtual FString GetFullTakePath(UE::CaptureManager::FTakeId InTakeId) const override;
+
+	//~ Begin ULiveLinkDeviceCapability_Ingest interface
+	virtual void UpdateTakeList_Implementation() override;
+	virtual void RunIngestTake(int32 InTakeId, const UIngestCapability_Options* InIngestOptions) override;
+	virtual void CancelIngest(int32 InTakeId) override;
+	//~ End ULiveLinkDeviceCapability_Ingest interface
+
+	//~ Begin ILiveLinkDeviceCapability_Connection interface
+	virtual ELiveLinkDeviceConnectionStatus GetConnectionStatus_Implementation() const override;
+	virtual FString GetHardwareId_Implementation() const override;
+	virtual bool SetHardwareId_Implementation(const FString& HardwareID) override;
+	virtual bool Connect_Implementation() override;
+	virtual bool Disconnect_Implementation() override;
+	//~ End ILiveLinkDeviceCapability_Connection interface
+
+	//~ Begin ILiveLinkDeviceCapability_Recording interface
+	virtual bool StartRecording_Implementation() override;
+	virtual bool StopRecording_Implementation() override;
+	virtual bool IsRecording_Implementation() const override;
+	//~ End ILiveLinkDeviceCapability_Recording interface
+
+	void HandleConnectionChanged(TSharedPtr<const UE::CaptureManager::FCaptureEvent> InEvent);
+	void HandleCPSStateUpdate(TSharedPtr<const UE::CaptureManager::FCaptureEvent> InEvent);
+	void HandleCPSEvent(TSharedPtr<const UE::CaptureManager::FCaptureEvent> InEvent);
+
+	void OnExportProgressReport(float InProgress,
+								UE::CaptureManager::FTaskProgress::FTask InTask);
+	void OnExportFinished(UE::CaptureManager::TProtocolResult<void> InResult,
+						  FString InTakeName,
+						  UE::CaptureManager::FTakeId InTakeId,
+						  TSharedPtr<UE::CaptureManager::FTaskProgress> InTaskProgress,
+						  TStrongObjectPtr<const UIngestCapability_Options> InIngestOptions);
+
+	TOptional<FTakeMetadata> ParseTake(const FString& InTakeDirectory, const FString& InTakeName);
+	FTakeMetadata ParseTakeMetadata(const UE::CaptureManager::FGetTakeMetadataResponse::FTakeObject& InTake);
+
+	void FetchThumbnail(const FString& InTakeName, UE::CaptureManager::FTakeId InTakeId);
+	void FetchThumbnails(TMap<FString, UE::CaptureManager::FTakeId> InNameToIdMap);
+
+	TSharedPtr<UE::CaptureManager::FCPSDevice> Device;
+
+	mutable FCriticalSection DownloadedTakesMutex;
+	TMap<UE::CaptureManager::FTakeId, FString> DownloadedTakes;
+
+	std::atomic_bool bIsRecording = false;
+	std::atomic_bool bIsConnecting = false;
+};
diff --git a/Engine/Plugins/VirtualProduction/CaptureManager/CaptureManagerDevices/Source/CPSLiveLinkDevice/Private/LiveLinkFaceDeviceModule.cpp b/Engine/Plugins/VirtualProduction/CaptureManager/CaptureManagerDevices/Source/CPSLiveLinkDevice/Private/LiveLinkFaceDeviceModule.cpp
new file mode 100644
--- /dev/null
+++ b/Engine/Plugins/VirtualProduction/CaptureManager/CaptureManagerDevices/Source/CPSLiveLinkDevice/Private/LiveLinkFaceDeviceModule.cpp
@@ -0,0 +1,43 @@
+// Copyright Epic Games, Inc. All Rights Reserved.
+
+#include "Modules/ModuleManager.h"
+
+#include "PropertyEditorModule.h"
+
+#include "LiveLinkFaceDevice.h"
+#include "Customizations/ToggleConnectActionCustomization.h"
+#include "Customizations/DeviceIpAddressCustomization.h"
+
+class FCPSLiveLinkDeviceModule : public IModuleInterface
+{
+public:
+
+	void StartupModule()
+	{
+		FPropertyEditorModule& PropertyModule = FModuleManager::LoadModuleChecked<FPropertyEditorModule>("PropertyEditor");
+		PropertyModule.RegisterCustomPropertyTypeLayout(FToggleConnectAction::StaticStruct()->GetFName(),
+														FOnGetPropertyTypeCustomizationInstance::CreateRaw(this, &FCPSLiveLinkDeviceModule::GetToggleConnectCustomization));
+		PropertyModule.RegisterCustomPropertyTypeLayout(FDeviceIpAddress::StaticStruct()->GetFName(),
+														FOnGetPropertyTypeCustomizationInstance::CreateRaw(this, &FCPSLiveLinkDeviceModule::GetDeviceIpAddressCustomization));
+	}
+
+	void ShutdownModule()
+	{
+		FPropertyEditorModule& PropertyModule = FModuleManager::LoadModuleChecked<FPropertyEditorModule>("PropertyEditor");
+		PropertyModule.UnregisterCustomPropertyTypeLayout(FToggleConnectAction::StaticStruct()->GetFName());
+	}
+
+private:
+
+	TSharedRef<IPropertyTypeCustomization> GetToggleConnectCustomization()
+	{
+		return MakeShared<FToggleConnectActionCustomization>();
+	}
+
+	TSharedRef<IPropertyTypeCustomization> GetDeviceIpAddressCustomization()
+	{
+		return MakeShared<FDeviceIpAddressCustomization>();
+	}
+};
+
+IMPLEMENT_MODULE(FCPSLiveLinkDeviceModule, CPSLiveLinkDevice);
\ No newline at end of file
diff --git a/Engine/Plugins/VirtualProduction/CaptureManager/CaptureManagerDevices/Source/CPSLiveLinkDevice/Private/Protocol/CPSDataStream.cpp b/Engine/Plugins/VirtualProduction/CaptureManager/CaptureManagerDevices/Source/CPSLiveLinkDevice/Private/Protocol/CPSDataStream.cpp
new file mode 100644
--- /dev/null
+++ b/Engine/Plugins/VirtualProduction/CaptureManager/CaptureManagerDevices/Source/CPSLiveLinkDevice/Private/Protocol/CPSDataStream.cpp
@@ -0,0 +1,51 @@
+// Copyright Epic Games, Inc. All Rights Reserved.
+
+#include "CPSDataStream.h"
+
+#include "Misc/FileHelper.h"
+#include "Misc/Paths.h"
+
+namespace UE::CaptureManager {
+
+FCPSDataStream::FCPSDataStream(FFileExportFinished InFileExportFinished)
+	: FileExportFinished(MoveTemp(InFileExportFinished))
+{
+}
+
+bool FCPSDataStream::StartFile(const FString& InTakeName, const FString& InFileName)
+{
+	checkf(Data.IsEmpty(), TEXT("Data buffer must be empty"));
+	return true;
+}
+
+bool FCPSDataStream::ProcessData(const FString& InTakeName, const FString& InFileName, const TConstArrayView<uint8>& InData)
+{
+	Data.Append(InData);
+	return true;
+}
+
+bool FCPSDataStream::FinishFile(const FString& InTakeName, const FString& InFileName, const TStaticArray<uint8, 16>& InHash)
+{
+	using namespace UE::CaptureManager;
+
+	FMD5 MD5;
+	MD5.Update(Data.GetData(), Data.Num());
+
+	FMD5Hash Hash;
+	Hash.Set(MD5);
+
+	if (FMemory::Memcmp(Hash.GetBytes(), InHash.GetData(), Hash.GetSize()) != 0)
+	{
+		FileExportFinished.ExecuteIfBound(InTakeName, FCaptureProtocolError(TEXT("Invalid file hash")));
+	}
+
+	FileExportFinished.ExecuteIfBound(InTakeName, MoveTemp(Data));
+
+	return true;
+}
+
+void FCPSDataStream::Finalize(UE::CaptureManager::TProtocolResult<void> InResult)
+{
+}
+
+}
\ No newline at end of file
diff --git a/Engine/Plugins/VirtualProduction/CaptureManager/CaptureManagerDevices/Source/CPSLiveLinkDevice/Private/Protocol/CPSDataStream.h b/Engine/Plugins/VirtualProduction/CaptureManager/CaptureManagerDevices/Source/CPSLiveLinkDevice/Private/Protocol/CPSDataStream.h
new file mode 100644
--- /dev/null
+++ b/Engine/Plugins/VirtualProduction/CaptureManager/CaptureManagerDevices/Source/CPSLiveLinkDevice/Private/Protocol/CPSDataStream.h
@@ -0,0 +1,30 @@
+// Copyright Epic Games, Inc. All Rights Reserved.
+
+#pragma once
+
+#include "ExportClient/ExportClient.h"
+
+namespace UE::CaptureManager {
+
+class FCPSDataStream : public UE::CaptureManager::FBaseStream
+{
+public:
+	using FData = TArray<uint8>;
+
+	DECLARE_DELEGATE_TwoParams(FFileExportFinished, const FString& InTakeName, UE::CaptureManager::TProtocolResult<FData> InDataResult);
+
+	FCPSDataStream(FFileExportFinished InFileExportFinished);
+
+	bool StartFile(const FString& InTakeName, const FString& InFileName) override;
+	bool ProcessData(const FString& InTakeName, const FString& InFileName, const TConstArrayView<uint8>& InData) override;
+	bool FinishFile(const FString& InTakeName, const FString& InFileName, const TStaticArray<uint8, 16>& InHash) override;
+
+	void Finalize(UE::CaptureManager::TProtocolResult<void> InResult) override;
+
+private:
+
+	FData Data;
+	FFileExportFinished FileExportFinished;
+};
+
+}
\ No newline at end of file
diff --git a/Engine/Plugins/VirtualProduction/CaptureManager/CaptureManagerDevices/Source/CPSLiveLinkDevice/Private/Protocol/CPSDevice.cpp b/Engine/Plugins/VirtualProduction/CaptureManager/CaptureManagerDevices/Source/CPSLiveLinkDevice/Private/Protocol/CPSDevice.cpp
new file mode 100644
--- /dev/null
+++ b/Engine/Plugins/VirtualProduction/CaptureManager/CaptureManagerDevices/Source/CPSLiveLinkDevice/Private/Protocol/CPSDevice.cpp
@@ -0,0 +1,443 @@
+// Copyright Epic Games, Inc. All Rights Reserved.
+
+#include "CPSDevice.h"
+
+#include "Control/Messages/Constants.h"
+
+#include "Misc/ScopeExit.h"
+
+#include "CaptureUtilsModule.h"
+
+DEFINE_LOG_CATEGORY_STATIC(LogCPSProtocolDevice, Log, All);
+
+#define CHECK_LOG_RETURN_VAL(InCondition, InFormat, ...) \
+	if (InCondition) \
+	{ \
+		FString Message = FString::Printf(InFormat, ##__VA_ARGS__); \
+		UE_LOG(LogCPSProtocolDevice, Warning, TEXT("%s"), *Message); \
+		return FCaptureProtocolError(MoveTemp(Message)); \
+	}
+
+#define CHECK_LOG_RETURN(InCondition, InFormat, ...) \
+	if (InCondition) \
+	{ \
+		UE_LOG(LogCPSProtocolDevice, Warning, InFormat, ##__VA_ARGS__); \
+		return; \
+	}
+
+namespace UE::CaptureManager
+{
+
+FConnectionStateChangedEvent::FConnectionStateChangedEvent(EState InConnectionState)
+	: FCaptureEvent(Name)
+	, ConnectionState(InConnectionState)
+{
+}
+
+FCPSEvent::FCPSEvent(TSharedPtr<FControlUpdate> InUpdateMessage)
+	: FCaptureEvent(Name)
+	, UpdateMessage(MoveTemp(InUpdateMessage))
+{
+}
+
+FCPSStateEvent::FCPSStateEvent(FGetStateResponse InGetStateResponse)
+	: FCaptureEvent(Name)
+	, GetStateResponse(MoveTemp(InGetStateResponse))
+{
+}
+
+TSharedPtr<FCPSDevice> FCPSDevice::MakeCPSDevice(FString InDeviceIpAddress, uint16 InDevicePort)
+{
+	TSharedPtr<FCPSDevice> CPSDevice = MakeShared<FCPSDevice>(FPrivateToken(), MoveTemp(InDeviceIpAddress), InDevicePort);
+	CPSDevice->InitializeDelegates();
+
+	return CPSDevice;
+}
+
+FCPSDevice::FCPSDevice(FCPSDevice::FPrivateToken,
+					   FString InDeviceIpAddress,
+					   uint16 InDevicePort)
+	: TimerManager(GetTimerManager())
+	, DeviceIpAddress(MoveTemp(InDeviceIpAddress))
+	, DeviceControlPort(InDevicePort)
+	, bIsConnected(false)
+{
+	RegisterEvent(FConnectionStateChangedEvent::Name); // Connection changed
+	RegisterEvent(FCPSStateEvent::Name); // State update
+	RegisterEvent(FCPSEvent::Name); // Other events
+}
+
+FCPSDevice::~FCPSDevice()
+{
+	Stop();
+}
+
+void FCPSDevice::InitializeDelegates()
+{
+	ConnThread = MakeUnique<FConnectionThread>(FConnectionThread::FOnProcess::CreateSP(this, &FCPSDevice::ConnectControlClient));
+
+	RegisterForAllEvents();
+
+	ControlMessenger.RegisterDisconnectHandler(FControlMessenger::FOnDisconnect::CreateSP(this, &FCPSDevice::OnDisconnect));
+}
+
+void FCPSDevice::InitiateConnect()
+{
+	PublishEvent<FConnectionStateChangedEvent>(FConnectionStateChangedEvent::EState::Connecting);
+	StartConnectTimer();
+}
+
+void FCPSDevice::Stop()
+{
+	bool bIsConnectedLocal = bIsConnected.exchange(false);
+
+	ConnThread->Empty();
+	TimerManager->RemoveTimer(ConnectTimerHandle);
+
+	ControlMessenger.Stop();
+
+	if (bIsConnectedLocal)
+	{
+		PublishEvent<FConnectionStateChangedEvent>(FConnectionStateChangedEvent::EState::Disconnected);
+	}
+}
+
+bool FCPSDevice::IsConnected() const
+{
+	return bIsConnected.load();
+}
+
+TProtocolResult<void> FCPSDevice::StartRecording(FString SlateName,
+												 uint16 TakeNumber,
+												 TOptional<FString> Subject,
+												 TOptional<FString> Scenario,
+												 TOptional<TArray<FString>> Tags)
+{
+	
+
+	if (!bIsConnected.load())
+	{
+		FString Message = FString::Format(TEXT("Device is not connected: {0}:{1}"), { DeviceIpAddress, FString::FromInt(DeviceControlPort) });
+		return FCaptureProtocolError(MoveTemp(Message));
+	}
+
+	FStartRecordingTakeRequest Request(MoveTemp(SlateName),
+									   TakeNumber,
+									   MoveTemp(Subject),
+									   MoveTemp(Scenario),
+									   MoveTemp(Tags));
+
+	TProtocolResult<FStartRecordingTakeResponse> Response = ControlMessenger.SendRequest(MoveTemp(Request));
+
+	if (Response.HasError())
+	{
+		return Response.StealError();
+	}
+
+	return ResultOk;
+}
+
+TProtocolResult<void> FCPSDevice::StopRecording()
+{
+	if (!bIsConnected.load())
+	{
+		FString Message = FString::Format(TEXT("Device is not connected: {0}:{1}"), { DeviceIpAddress, FString::FromInt(DeviceControlPort) });
+		return FCaptureProtocolError(MoveTemp(Message));
+	}
+
+	FStopRecordingTakeRequest Request;
+
+	TProtocolResult<FStopRecordingTakeResponse> Response = ControlMessenger.SendRequest(MoveTemp(Request));
+
+	if (Response.HasError())
+	{
+		return Response.StealError();
+	}
+
+	return ResultOk;
+}
+
+TProtocolResult<TArray<FGetTakeMetadataResponse::FTakeObject>> FCPSDevice::FetchTakeList()
+{
+	if (!bIsConnected.load())
+	{
+		FString Message = FString::Format(TEXT("Device is not connected: {0}:{1}"), { DeviceIpAddress, FString::FromInt(DeviceControlPort) });
+		return FCaptureProtocolError(MoveTemp(Message));
+	}
+
+	TProtocolResult<FGetTakeListResponse> TakeListRequestResult = ControlMessenger.SendRequest(FGetTakeListRequest());
+	CHECK_LOG_RETURN_VAL(TakeListRequestResult.HasError(), TEXT("Failed to fetch the take list from the remote host."));
+
+	FGetTakeListResponse TakeListResponse = TakeListRequestResult.StealValue();
+
+	FGetTakeMetadataRequest TakeMetadataRequest(TakeListResponse.GetNames());
+
+	TProtocolResult<FGetTakeMetadataResponse> TakeMetadataRequestResult = ControlMessenger.SendRequest(MoveTemp(TakeMetadataRequest));
+	CHECK_LOG_RETURN_VAL(TakeMetadataRequestResult.HasError(), TEXT("Failed to fetch the take metadata from the remote host. Note: Please ensure you are using compatible versions of LLF and UE"));
+
+	FGetTakeMetadataResponse TakeMetadataResponse = TakeMetadataRequestResult.StealValue();
+
+	return TakeMetadataResponse.GetTakes();
+}
+
+TProtocolResult<FGetTakeMetadataResponse::FTakeObject> FCPSDevice::FetchTake(const FString& InTakeName)
+{
+	if (!bIsConnected.load())
+	{
+		FString Message = FString::Format(TEXT("Device is not connected: {0}:{1}"), { DeviceIpAddress, FString::FromInt(DeviceControlPort) });
+		return FCaptureProtocolError(MoveTemp(Message));
+	}
+
+	FGetTakeMetadataRequest TakeMetadataRequest({ InTakeName });
+
+	TProtocolResult<FGetTakeMetadataResponse> TakeMetadataRequestResult = ControlMessenger.SendRequest(MoveTemp(TakeMetadataRequest));
+	CHECK_LOG_RETURN_VAL(TakeMetadataRequestResult.HasError(), TEXT("Failed to fetch the take metadata from the remote host. Please ensure you are using compatible versions of LLF and UE"));
+
+	FGetTakeMetadataResponse TakeMetadataResponse = TakeMetadataRequestResult.StealValue();
+
+	TArray<FGetTakeMetadataResponse::FTakeObject> Take = TakeMetadataResponse.GetTakes();
+
+	if (Take.IsEmpty())
+	{
+		return FCaptureProtocolError(FString::Format(TEXT("Failed to obtain the take with the specified name: {0}"), { InTakeName }));
+	}
+
+	check(Take.Num() == 1);
+
+	return Take[0];
+}
+
+void FCPSDevice::AddTakeMetadata(FTakeId InId, FGetTakeMetadataResponse::FTakeObject InTake)
+{
+	FScopeLock Lock(&Mutex);
+	TakeMetadata.Add(InId, MoveTemp(InTake));
+}
+
+void FCPSDevice::RemoveTakeMetadata(FTakeId InId)
+{
+	FScopeLock Lock(&Mutex);
+	TakeMetadata.Remove(InId);
+}
+
+FGetTakeMetadataResponse::FTakeObject FCPSDevice::GetTake(FTakeId InId)
+{
+	FScopeLock Lock(&Mutex);
+	return TakeMetadata[InId];
+}
+
+FTakeId FCPSDevice::GetTakeId(const FString& InTakeName)
+{
+	FScopeLock Lock(&Mutex);
+
+	for (const TPair<FTakeId, FGetTakeMetadataResponse::FTakeObject>& TakeElement : TakeMetadata)
+	{
+		if (TakeElement.Value.Name == InTakeName)
+		{
+			return TakeElement.Key;
+		}
+	}
+
+	return INDEX_NONE;
+}
+
+void FCPSDevice::StartExport(FTakeId InTakeId, TUniquePtr<FBaseStream> InStream)
+{
+	check(ExportClient);
+
+	const FGetTakeMetadataResponse::FTakeObject& Take = GetTake(InTakeId);
+
+	FExportClient::FTakeFileArray TakeFiles;
+	TakeFiles.Reserve(Take.Files.Num());
+
+	for (const FGetTakeMetadataResponse::FFileObject& File : Take.Files)
+	{
+		TakeFiles.Add({ File.Name, File.Length, 0 });
+	}
+
+	FExportClient::FTaskId TaskId = ExportClient->ExportTakeFiles(Take.Name, MoveTemp(TakeFiles), MoveTemp(InStream));
+	IdMap.Add(InTakeId, TaskId);
+}
+
+void FCPSDevice::CancelExport(FTakeId InTakeId)
+{
+	if (FExportClient::FTaskId* TaskId = IdMap.Find(InTakeId))
+	{
+		ExportClient->AbortExport(*TaskId);
+	}
+}
+
+void FCPSDevice::FetchThumbnail(FTakeId InTakeId, TUniquePtr<FBaseStream> InStream)
+{
+	static constexpr FStringView FileName = TEXTVIEW("thumbnail.jpg");
+
+	FScopeLock Lock(&Mutex);
+	if (FGetTakeMetadataResponse::FTakeObject* TakeObjectPtr = TakeMetadata.Find(InTakeId))
+	{
+		// Intentional copy
+		const FGetTakeMetadataResponse::FTakeObject TakeObject = *TakeObjectPtr;
+
+		// Unlock earlier as there is nothing below that needs protecting
+		Lock.Unlock();
+
+		uint64 Length = 0;
+		for (const FGetTakeMetadataResponse::FFileObject& FileObject : TakeObject.Files)
+		{
+			if (FileObject.Name == FileName)
+			{
+				Length = FileObject.Length;
+			}
+		}
+
+		if (Length == 0)
+		{
+			return;
+		}
+
+		FTakeFile TakeFile;
+		TakeFile.FileName = FileName;
+		TakeFile.Length = Length;
+		TakeFile.Offset = 0;
+
+		ExportClient->ExportTakeFiles(TakeObject.Name, { MoveTemp(TakeFile) }, MoveTemp(InStream));
+	}
+}
+
+void FCPSDevice::FetchThumbnails(TUniquePtr<FBaseStream> InStream)
+{
+	static constexpr FStringView FileName = TEXTVIEW("thumbnail.jpg");
+
+	TMap<FString, FExportClient::FTakeFileArray> TakeFileArrayMap;
+
+	FScopeLock Lock(&Mutex);
+
+	for (const TPair<FTakeId, FGetTakeMetadataResponse::FTakeObject>& TakePair : TakeMetadata)
+	{
+		const FGetTakeMetadataResponse::FTakeObject& TakeObject = TakePair.Value;
+
+		uint64 Length = 0;
+		for (const FGetTakeMetadataResponse::FFileObject& FileObject : TakeObject.Files)
+		{
+			if (FileObject.Name == FileName)
+			{
+				Length = FileObject.Length;
+			}
+		}
+
+		if (Length == 0)
+		{
+			continue;
+		}
+
+		FTakeFile TakeFile;
+		TakeFile.FileName = FileName;
+		TakeFile.Length = Length;
+		TakeFile.Offset = 0;
+
+		TakeFileArrayMap.Add(TakeObject.Name, { MoveTemp(TakeFile) });
+	}
+
+	// Unlock earlier as there is no need to lock the ExporFiles function
+	Lock.Unlock();
+
+	ExportClient->ExportFiles(MoveTemp(TakeFileArrayMap), MoveTemp(InStream));
+}
+
+void FCPSDevice::ConnectControlClient(FEmpty)
+{
+	if (bIsConnected.load())
+	{
+		return;
+	}
+
+	ON_SCOPE_EXIT
+	{
+		if (!bIsConnected.load())
+		{
+			ControlMessenger.Stop();
+		}
+	};
+
+	if (ControlMessenger.Start(DeviceIpAddress, DeviceControlPort).HasError())
+	{
+		UE_LOG(LogCPSProtocolDevice, Verbose, TEXT("Failed to connect to %s:%d"), *DeviceIpAddress, DeviceControlPort);
+		return;
+	}
+
+	CHECK_LOG_RETURN(ControlMessenger.StartSession().HasError(), TEXT("Failed to start session for the Control client. Please ensure you are using compatible versions of LLF and UE"));
+
+	CHECK_LOG_RETURN(ControlMessenger.SendRequest(FSubscribeRequest()).HasError(), TEXT("Failed to subscribe to events for Control client."));
+
+	TProtocolResult<FGetServerInformationResponse> GetServerInformationResult = ControlMessenger.GetServerInformation();
+	CHECK_LOG_RETURN(GetServerInformationResult.HasError(), TEXT("Failed to fetch the Control server information"));
+	FGetServerInformationResponse GetServerInformationResponse = GetServerInformationResult.StealValue();
+
+	ExportClient = MakeUnique<FExportClient>(DeviceIpAddress, GetServerInformationResponse.GetExportPort());
+
+	TProtocolResult<FGetStateResponse> GetStateResult = ControlMessenger.SendRequest(FGetStateRequest());
+	CHECK_LOG_RETURN(GetStateResult.HasError(), TEXT("Failed to fetch the current state of the Control server"));
+
+	FGetStateResponse GetStateResponse = GetStateResult.StealValue();
+	PublishEvent<FCPSStateEvent>(MoveTemp(GetStateResponse));
+
+	PublishEvent<FConnectionStateChangedEvent>(FConnectionStateChangedEvent::EState::Connected);
+
+	bIsConnected.store(true);
+
+	UE_LOG(LogCPSProtocolDevice, Display, TEXT("Connected to the CPS Device: %s:%d"), *DeviceIpAddress, DeviceControlPort);
+}
+
+void FCPSDevice::RegisterForAllEvents()
+{
+	ControlMessenger.RegisterUpdateHandler(CPS::AddressPaths::GTakeAdded, FControlUpdate::FOnUpdateMessage::CreateSP(this, &FCPSDevice::OnCPSEvent));
+	ControlMessenger.RegisterUpdateHandler(CPS::AddressPaths::GTakeRemoved, FControlUpdate::FOnUpdateMessage::CreateSP(this, &FCPSDevice::OnCPSEvent));
+	ControlMessenger.RegisterUpdateHandler(CPS::AddressPaths::GTakeUpdated, FControlUpdate::FOnUpdateMessage::CreateSP(this, &FCPSDevice::OnCPSEvent));
+	ControlMessenger.RegisterUpdateHandler(CPS::AddressPaths::GRecordingStatus, FControlUpdate::FOnUpdateMessage::CreateSP(this, &FCPSDevice::OnCPSEvent));
+	ControlMessenger.RegisterUpdateHandler(CPS::AddressPaths::GDiskCapacity, FControlUpdate::FOnUpdateMessage::CreateSP(this, &FCPSDevice::OnCPSEvent));
+	ControlMessenger.RegisterUpdateHandler(CPS::AddressPaths::GBattery, FControlUpdate::FOnUpdateMessage::CreateSP(this, &FCPSDevice::OnCPSEvent));
+	ControlMessenger.RegisterUpdateHandler(CPS::AddressPaths::GThermalState, FControlUpdate::FOnUpdateMessage::CreateSP(this, &FCPSDevice::OnCPSEvent));
+}
+
+void FCPSDevice::OnCPSEvent(TSharedPtr<FControlUpdate> InUpdateMessage)
+{
+	PublishEvent<FCPSEvent>(MoveTemp(InUpdateMessage));
+}
+
+void FCPSDevice::StartConnectTimer(float InDelay)
+{
+	if (!ConnectTimerHandle.IsValid())
+	{
+		ConnectTimerHandle = TimerManager->AddTimer(FTimerDelegate::CreateSP(this, &FCPSDevice::OnConnectTick), ConnectInterval, true, InDelay);
+	}
+}
+
+void FCPSDevice::OnConnectTick()
+{
+	if (!bIsConnected.load())
+	{
+		UE_LOG(LogCPSProtocolDevice, Display, TEXT("Connecting to the CPS Device: %s:%d"), *DeviceIpAddress, DeviceControlPort);
+		ConnThread->Add(FEmpty()); // Initiate connect
+	}
+	else
+	{
+		TimerManager->RemoveTimer(ConnectTimerHandle);
+	}
+}
+
+void FCPSDevice::OnDisconnect(const FString& InCause)
+{
+	if (bIsConnected.exchange(false))
+	{
+		PublishEvent<FConnectionStateChangedEvent>(FConnectionStateChangedEvent::EState::Disconnected);
+	}
+
+	PublishEvent<FConnectionStateChangedEvent>(FConnectionStateChangedEvent::EState::Connecting);
+	StartConnectTimer(ConnectInterval);
+}
+
+TSharedRef<FCaptureTimerManager> FCPSDevice::GetTimerManager()
+{
+	FCaptureUtilsModule& Module = FModuleManager::LoadModuleChecked<FCaptureUtilsModule>(TEXT("CaptureUtils"));
+	return Module.GetTimerManager();
+}
+
+}
\ No newline at end of file
diff --git a/Engine/Plugins/VirtualProduction/CaptureManager/CaptureManagerDevices/Source/CPSLiveLinkDevice/Private/Protocol/CPSDevice.h b/Engine/Plugins/VirtualProduction/CaptureManager/CaptureManagerDevices/Source/CPSLiveLinkDevice/Private/Protocol/CPSDevice.h
new file mode 100644
--- /dev/null
+++ b/Engine/Plugins/VirtualProduction/CaptureManager/CaptureManagerDevices/Source/CPSLiveLinkDevice/Private/Protocol/CPSDevice.h
@@ -0,0 +1,143 @@
+// Copyright Epic Games, Inc. All Rights Reserved.
+
+#pragma once
+
+#include "Control/ControlMessenger.h"
+
+#include "ExportClient/ExportClient.h"
+
+#include "Async/EventSourceUtils.h"
+#include "Misc/Optional.h"
+
+#include "Async/CaptureTimerManager.h"
+
+#include "Ingest/IngestCapability_TakeInformation.h"
+
+namespace UE::CaptureManager
+{
+
+struct FConnectionStateChangedEvent : public FCaptureEvent
+{
+	inline static const FString Name = TEXT("ConnectionStateChanged");
+
+	enum class EState
+	{
+		Unknown = 0,
+		Connecting,
+		Connected,
+		Disconnected
+	};
+
+	FConnectionStateChangedEvent(EState InConnectionState);
+
+	EState ConnectionState;
+};
+
+struct FCPSEvent : public FCaptureEvent
+{
+	inline static const FString Name = TEXT("CPSEvent");
+
+	FCPSEvent(TSharedPtr<FControlUpdate> InUpdateMessage);
+
+	TSharedPtr<FControlUpdate> UpdateMessage;
+};
+
+struct FCPSStateEvent : public FCaptureEvent
+{
+	inline static const FString Name = TEXT("CPSStateEvent");
+
+	FCPSStateEvent(FGetStateResponse InGetStateResponse);
+
+	FGetStateResponse GetStateResponse;
+};
+
+class FCPSDevice :
+	public FCaptureEventSource,
+	public TSharedFromThis<FCPSDevice>
+{
+private:
+
+	struct FPrivateToken
+	{
+		explicit FPrivateToken() = default;
+	};
+
+public:
+
+	static TSharedPtr<FCPSDevice> MakeCPSDevice(FString InDeviceIpAddress, uint16 InDevicePort);
+
+	FCPSDevice(FPrivateToken,
+			   FString InDeviceIpAddress,
+			   uint16 InDevicePort);
+
+	virtual ~FCPSDevice() override;
+
+	void InitiateConnect();
+	void Stop();
+
+	bool IsConnected() const;
+
+	TProtocolResult<void> StartRecording(FString SlateName,
+										 uint16 TakeNumber,
+										 TOptional<FString> Subject = TOptional<FString>(),
+										 TOptional<FString> Scenario = TOptional<FString>(),
+										 TOptional<TArray<FString>> Tags = TOptional<TArray<FString>>());
+
+	TProtocolResult<void> StopRecording();
+
+	TProtocolResult<TArray<FGetTakeMetadataResponse::FTakeObject>> FetchTakeList();
+	TProtocolResult<FGetTakeMetadataResponse::FTakeObject> FetchTake(const FString& InTakeName);
+
+	void AddTakeMetadata(FTakeId InId, FGetTakeMetadataResponse::FTakeObject InTake);
+	void RemoveTakeMetadata(FTakeId InId);
+	FGetTakeMetadataResponse::FTakeObject GetTake(FTakeId InId);
+	FTakeId GetTakeId(const FString& InTakeName);
+
+	void StartExport(FTakeId InTakeId, TUniquePtr<FBaseStream> InStream);
+	void CancelExport(FTakeId InTakeId);
+
+	void FetchThumbnail(FTakeId InTakeId, TUniquePtr<FBaseStream> InStream);
+	void FetchThumbnails(TUniquePtr<FBaseStream> InStream);
+
+private:
+
+	inline static const float ConnectInterval = 5.0f;
+
+	void InitializeDelegates();
+
+	struct FEmpty
+	{
+	};
+
+	void ConnectControlClient(FEmpty);
+
+	void RegisterForAllEvents();
+	void OnCPSEvent(TSharedPtr<FControlUpdate> InUpdateMessage);
+
+	void StartConnectTimer(float InDelay = 0.0f);
+	void OnConnectTick();
+	void OnDisconnect(const FString& InCause);
+
+	static TSharedRef<FCaptureTimerManager> GetTimerManager();
+
+	TSharedRef<FCaptureTimerManager> TimerManager;
+
+	FString DeviceIpAddress;
+	uint16 DeviceControlPort;
+
+	std::atomic_bool bIsConnected;
+
+	using FConnectionThread = TQueueRunner<FEmpty>;
+	TUniquePtr<FConnectionThread> ConnThread;
+	FControlMessenger ControlMessenger;
+	FCaptureTimerManager::FTimerHandle ConnectTimerHandle;
+
+	TUniquePtr<FExportClient> ExportClient;
+
+	FCriticalSection Mutex;
+	TMap<FTakeId, FGetTakeMetadataResponse::FTakeObject> TakeMetadata;
+
+	TMap<FTakeId, FExportClient::FTaskId> IdMap;
+};
+
+}
\ No newline at end of file
diff --git a/Engine/Plugins/VirtualProduction/CaptureManager/CaptureManagerDevices/Source/CPSLiveLinkDevice/Private/Protocol/CPSFileStream.cpp b/Engine/Plugins/VirtualProduction/CaptureManager/CaptureManagerDevices/Source/CPSLiveLinkDevice/Private/Protocol/CPSFileStream.cpp
new file mode 100644
--- /dev/null
+++ b/Engine/Plugins/VirtualProduction/CaptureManager/CaptureManagerDevices/Source/CPSLiveLinkDevice/Private/Protocol/CPSFileStream.cpp
@@ -0,0 +1,93 @@
+// Copyright Epic Games, Inc. All Rights Reserved.
+
+#include "CPSFileStream.h"
+
+#include "Misc/FileHelper.h"
+#include "Misc/Paths.h"
+
+namespace UE::CaptureManager
+{
+
+FCPSFileStream::FCPSFileStream(FString InBaseDir, uint64 InSize)
+	: BaseDir(MoveTemp(InBaseDir))
+	, TotalExportExpectedSize(InSize)
+	, TotalExportArrivedSize(0)
+{
+}
+
+bool FCPSFileStream::StartFile(const FString& InTakeName, const FString& InFileName)
+{
+	FString FilePathToBeSaved = FPaths::Combine(BaseDir, InTakeName, InFileName);
+
+	checkf(!Writer, TEXT("Writer must be nullptr when new file arrives"));
+	Writer = TUniquePtr<FArchive>(IFileManager::Get().CreateFileWriter(*FilePathToBeSaved, EFileWrite::FILEWRITE_None));
+
+	checkf(!MD5Generator, TEXT("MD5 generator must be nullptr when new file arrives"));
+	MD5Generator = MakeUnique<FMD5>();
+
+	return true;
+}
+
+bool FCPSFileStream::ProcessData(const FString& InTakeName, const FString& InFileName, const TConstArrayView<uint8>& InData)
+{
+	Writer->Serialize(const_cast<uint8*>(InData.GetData()), InData.Num());
+	ReportProgressStep(InFileName, InData.Num());
+
+	MD5Generator->Update(InData.GetData(), InData.Num());
+
+	return true;
+}
+
+bool FCPSFileStream::FinishFile(const FString& InTakeName, const FString& InFileName, const TStaticArray<uint8, 16>& InHash)
+{
+	checkf(Writer, TEXT("Writer must be valid when file has arrived"));
+	Writer->Close();
+	Writer = nullptr;
+
+	checkf(MD5Generator, TEXT("MD5 Generator must be valid when file has arrived"));
+
+	static constexpr uint32 HashSize = 16;
+	TStaticArray<uint8, HashSize> Hash = MakeUniformStaticArray<uint8, HashSize>(0);
+
+	MD5Generator->Final(Hash.GetData());
+	MD5Generator = nullptr;
+
+	if (FMemory::Memcmp(Hash.GetData(), InHash.GetData(), Hash.Num()) != 0)
+	{
+		Finalize(UE::CaptureManager::FCaptureProtocolError(TEXT("Invalid hash file")));
+
+		return false;
+	}
+
+	return true;
+}
+
+void FCPSFileStream::Finalize(UE::CaptureManager::TProtocolResult<void> InResult)
+{
+	// Close the writer if not closed
+	if (Writer)
+	{
+		Writer->Close();
+		Writer = nullptr;
+	}
+
+	OnExportFinished.ExecuteIfBound(MoveTemp(InResult));
+}
+
+void FCPSFileStream::SetExportFinished(FExportFinished InExportFinished)
+{
+	OnExportFinished = MoveTemp(InExportFinished);
+}
+
+void FCPSFileStream::SetProgressHandler(FReportProgress InReportProgress)
+{
+	OnReportProgress = MoveTemp(InReportProgress);
+}
+
+void FCPSFileStream::ReportProgressStep(const FString& InFileName, const uint32 InArrivedSize)
+{
+	TotalExportArrivedSize += InArrivedSize;
+	OnReportProgress.ExecuteIfBound((float) TotalExportArrivedSize / (float) TotalExportExpectedSize);
+}
+
+}
\ No newline at end of file
diff --git a/Engine/Plugins/VirtualProduction/CaptureManager/CaptureManagerDevices/Source/CPSLiveLinkDevice/Private/Protocol/CPSFileStream.h b/Engine/Plugins/VirtualProduction/CaptureManager/CaptureManagerDevices/Source/CPSLiveLinkDevice/Private/Protocol/CPSFileStream.h
new file mode 100644
--- /dev/null
+++ b/Engine/Plugins/VirtualProduction/CaptureManager/CaptureManagerDevices/Source/CPSLiveLinkDevice/Private/Protocol/CPSFileStream.h
@@ -0,0 +1,45 @@
+// Copyright Epic Games, Inc. All Rights Reserved.
+
+#pragma once
+
+#include "ExportClient/ExportClient.h"
+
+namespace UE::CaptureManager
+{
+
+class FCPSFileStream : public UE::CaptureManager::FBaseStream
+{
+public:
+
+	DECLARE_DELEGATE_OneParam(FReportProgress, float InProgress);
+	DECLARE_DELEGATE_OneParam(FExportFinished, UE::CaptureManager::TProtocolResult<void> InResult);
+
+	FCPSFileStream(FString InBaseDir, uint64 InSize);
+
+	virtual bool StartFile(const FString& InTakeName, const FString& InFileName) override;
+
+	virtual bool ProcessData(const FString& InTakeName, const FString& InFileName, const TConstArrayView<uint8>& InData) override;
+	virtual bool FinishFile(const FString& InTakeName, const FString& InFileName, const TStaticArray<uint8, 16>& InHash) override;
+
+	virtual void Finalize(UE::CaptureManager::TProtocolResult<void> InResult) override;
+
+	void SetExportFinished(FExportFinished InExportFinished);
+	void SetProgressHandler(FReportProgress InReportProgress);
+
+private:
+
+	void ReportProgressStep(const FString& InFileName, const uint32 InArrivedSize);
+
+	FString BaseDir;
+	TUniquePtr<FArchive> Writer;
+	TUniquePtr<FMD5> MD5Generator;
+
+	uint64 TotalExportExpectedSize;
+	uint64 TotalExportArrivedSize;
+
+	FExportFinished OnExportFinished;
+	FReportProgress OnReportProgress;
+};
+
+}
+
diff --git a/Engine/Plugins/VirtualProduction/CaptureManager/CaptureManagerDevices/Source/FilesystemLiveLinkDevice/FilesystemLiveLinkDevice.Build.cs b/Engine/Plugins/VirtualProduction/CaptureManager/CaptureManagerDevices/Source/FilesystemLiveLinkDevice/FilesystemLiveLinkDevice.Build.cs
new file mode 100644
--- /dev/null
+++ b/Engine/Plugins/VirtualProduction/CaptureManager/CaptureManagerDevices/Source/FilesystemLiveLinkDevice/FilesystemLiveLinkDevice.Build.cs
@@ -0,0 +1,44 @@
+// Copyright Epic Games, Inc. All Rights Reserved.
+
+using UnrealBuildTool;
+using System.IO;
+
+public class FilesystemLiveLinkDevice : ModuleRules
+{
+	public FilesystemLiveLinkDevice(ReadOnlyTargetRules Target) : base(Target)
+	{
+		PCHUsage = PCHUsageMode.UseExplicitOrSharedPCHs;
+		ShortName = "FilesystemLLD";
+
+		PublicDependencyModuleNames.AddRange(new string[]
+		{
+			"CaptureUtils",
+			"Core"
+		});
+
+		PrivateDependencyModuleNames.AddRange(new string[]
+		{
+			"Engine",
+			"Json",
+			"CaptureManagerMediaRW",
+			"IngestLiveLinkDevice",
+			"LiveLinkDevice",
+			"LiveLinkCapabilities",
+			"Media",
+			"LiveLinkHubCaptureMessaging",
+			"LiveLinkHubExportClient",
+			"LiveLinkHub",
+			"LiveLinkFaceMetadata",
+			"ToolWidgets",
+			"Slate",
+			"SlateCore",
+			"StereoCameraMetadata",
+			"CaptureManagerTakeMetadata",
+			"CaptureManagerSettings",
+			"CoreUObject",
+			"UnrealEd"
+		});
+
+		CppCompileWarningSettings.UnsafeTypeCastWarningLevel = WarningLevel.Error;
+	}
+}
diff --git a/Engine/Plugins/VirtualProduction/CaptureManager/CaptureManagerDevices/Source/FilesystemLiveLinkDevice/Private/FilesystemLiveLinkDevice.cpp b/Engine/Plugins/VirtualProduction/CaptureManager/CaptureManagerDevices/Source/FilesystemLiveLinkDevice/Private/FilesystemLiveLinkDevice.cpp
new file mode 100644
--- /dev/null
+++ b/Engine/Plugins/VirtualProduction/CaptureManager/CaptureManagerDevices/Source/FilesystemLiveLinkDevice/Private/FilesystemLiveLinkDevice.cpp
@@ -0,0 +1,284 @@
+// Copyright Epic Games, Inc. All Rights Reserved.
+
+#include "FilesystemLiveLinkDevice.h"
+
+#include "HAL/FileManager.h"
+#include "Settings/CaptureManagerSettings.h"
+
+#include "Utils/GenericLiveLinkDeviceLog.h"
+
+#include "LiveLinkFaceMetadata.h"
+#include "StereoCameraMetadataParseUtils.h"
+
+#include "Utils/IngestLiveLinkDeviceUtils.h"
+
+DEFINE_LOG_CATEGORY(LogGenericLiveLinkArchiveDevice);
+
+const UFilesystemLiveLinkDeviceSettings* UFilesystemLiveLinkDevice::GetSettings() const
+{
+	return CastChecked<UFilesystemLiveLinkDeviceSettings>(Settings);
+}
+
+TSubclassOf<ULiveLinkDeviceSettings> UFilesystemLiveLinkDevice::StaticGetSettingsClass() const
+{
+	return UFilesystemLiveLinkDeviceSettings::StaticClass();
+}
+
+FText UFilesystemLiveLinkDevice::GetDisplayName() const
+{
+	return FText::FromString(GetSettings()->DisplayName);
+}
+
+EDeviceHealth UFilesystemLiveLinkDevice::GetDeviceHealth() const
+{
+	return EDeviceHealth::Nominal;
+}
+
+FText UFilesystemLiveLinkDevice::GetHealthText() const
+{
+	return FText::FromString("Nominal");
+}
+
+void UFilesystemLiveLinkDevice::OnSettingChanged(const FPropertyChangedEvent& InPropertyChangedEvent)
+{
+	Super::OnSettingChanged(InPropertyChangedEvent);
+
+	if (GET_MEMBER_NAME_CHECKED(UFilesystemLiveLinkDeviceSettings, StoragePath) == InPropertyChangedEvent.GetMemberPropertyName())
+	{
+		FString Path = GetSettings()->StoragePath.Path;
+		if (!Path.IsEmpty() && FPaths::DirectoryExists(Path))
+		{
+			SetConnectionStatus(ELiveLinkDeviceConnectionStatus::Connected);
+		}
+		else
+		{
+			SetConnectionStatus(ELiveLinkDeviceConnectionStatus::Disconnected);
+		}
+	}
+}
+
+FString UFilesystemLiveLinkDevice::GetFullTakePath(UE::CaptureManager::FTakeId InTakeId) const
+{
+	FString StoragePath = GetSettings()->StoragePath.Path;
+
+	if (const FString* RelativeTakePath = RelativeTakePaths.Find(InTakeId))
+	{
+		return FPaths::Combine(StoragePath, *RelativeTakePath);
+	}
+
+	return FString();
+}
+
+void UFilesystemLiveLinkDevice::RunIngestTake(int32 InTakeId, const UIngestCapability_Options* InIngestOptions)
+{
+	AsyncTask(ENamedThreads::AnyThread, [this,
+			  TakeId = InTakeId,
+			  IngestOptions = InIngestOptions]()
+	{
+		static constexpr uint32 NumberOfTasks = 2; // Convert, Upload
+
+		using namespace UE::CaptureManager;
+		TSharedPtr<FTaskProgress> TaskProgress
+			= MakeShared<FTaskProgress>(NumberOfTasks, FTaskProgress::FProgressReporter::CreateLambda([this, TakeId](double InProgress)
+		{
+			ExecuteProcessProgressReporter(TakeId, InProgress);
+		}));
+
+		Super::IngestTake(TakeId, IngestOptions, MoveTemp(TaskProgress));
+	});
+}
+
+void UFilesystemLiveLinkDevice::UpdateTakeList_Implementation()
+{
+	AsyncTask(ENamedThreads::AnyThread, [this]()
+	{
+		RemoveAllTakes();
+
+		FString StoragePath = GetSettings()->StoragePath.Path;
+
+		TArray<FString> TakeMetadataFiles;
+		const bool bIterationResult = IFileManager::Get().IterateDirectoryRecursively(*StoragePath, [&](const TCHAR* InFileNameOrDirectory, bool bInIsDirectory)
+		{
+			if (!bInIsDirectory)
+			{
+				const FString CurrentFileName = FPaths::GetCleanFilename(InFileNameOrDirectory);
+				if (CurrentFileName == TEXT("take.json"))
+				{
+					TakeMetadataFiles.Add(InFileNameOrDirectory);
+				}
+				else if (FPaths::GetExtension(CurrentFileName) == FTakeMetadata::FileExtension)
+				{
+					TakeMetadataFiles.Add(InFileNameOrDirectory);
+				}
+			}
+
+			return true;
+		});
+
+		for (const FString& CurrentTakeFile : TakeMetadataFiles)
+		{
+			TOptional<FTakeMetadata> TakeInfoResult = ReadTake(CurrentTakeFile);
+
+			if (TakeInfoResult.IsSet())
+			{
+				int32 TakeId = AddTake(TakeInfoResult.GetValue());
+				FString CurrentDirectory = FPaths::GetPath(CurrentTakeFile);
+
+				FString TakePath = CurrentDirectory.Right(CurrentDirectory.Len() - StoragePath.Len());
+				FPaths::NormalizeDirectoryName(TakePath); // Removes trailing slash
+				RelativeTakePaths.Add(TakeId, MoveTemp(TakePath));
+			}
+		}
+
+		ExecuteUpdateTakeListCallback(Execute_GetTakeIdentifiers(this));
+	});
+}
+
+ELiveLinkDeviceConnectionStatus UFilesystemLiveLinkDevice::GetConnectionStatus_Implementation() const
+{
+	const UFilesystemLiveLinkDeviceSettings* DeviceSettings = GetSettings();
+	FString Path = DeviceSettings->StoragePath.Path;
+
+	if (!Path.IsEmpty() && FPaths::DirectoryExists(Path))
+	{
+		return ELiveLinkDeviceConnectionStatus::Connected;
+	}
+
+	return ELiveLinkDeviceConnectionStatus::Disconnected;
+}
+
+FString UFilesystemLiveLinkDevice::GetHardwareId_Implementation() const
+{
+	return FPlatformMisc::GetDeviceId();
+}
+
+bool UFilesystemLiveLinkDevice::SetHardwareId_Implementation(const FString& HardwareID)
+{
+	return false;
+}
+
+bool UFilesystemLiveLinkDevice::Connect_Implementation()
+{
+	const UFilesystemLiveLinkDeviceSettings* DeviceSettings = GetSettings();
+	FString Path = DeviceSettings->StoragePath.Path;
+
+	if (!Path.IsEmpty() && FPaths::DirectoryExists(Path))
+	{
+		SetConnectionStatus(ELiveLinkDeviceConnectionStatus::Connected);
+		return true;
+	}
+	
+	return false;
+}
+
+bool UFilesystemLiveLinkDevice::Disconnect_Implementation()
+{
+	SetConnectionStatus(ELiveLinkDeviceConnectionStatus::Disconnected);
+	return true;
+}
+
+TOptional<FTakeMetadata> UFilesystemLiveLinkDevice::ReadTake(const FString& InCurrentTakeFile)
+{
+	using namespace UE::CaptureManager;
+
+	FTakeMetadataParser TakeMetadataParser;
+	TValueOrError<FTakeMetadata, FTakeMetadataParserError> TakeMetadataResult = TakeMetadataParser.Parse(InCurrentTakeFile);
+	FString CurrentDirectory = FPaths::GetPath(InCurrentTakeFile);
+
+	if (TakeMetadataResult.HasError())
+	{
+		FTakeMetadataParserError ParserError = TakeMetadataResult.StealError();
+		FString OriginName = UE::CaptureManager::ErrorOriginToString(ParserError.Origin);
+
+		FString Message =
+			FString::Format(TEXT("Unable to parse take metadata file - {0} (Error origin: {1}): {2}"),
+							{
+								InCurrentTakeFile,
+								ParserError.Message.ToString(),
+								OriginName
+							});
+
+		UE_LOG(LogGenericLiveLinkArchiveDevice, Warning, TEXT("%s"), *Message);
+
+		UE_LOG(LogGenericLiveLinkArchiveDevice, Display, TEXT("Checking backwards compatible take metadata formats"), *CurrentDirectory);
+
+		TArray<FText> ValidationFailures;
+
+		UE_LOG(LogGenericLiveLinkArchiveDevice, Display, TEXT("Checking directory (%s) for pre UE 5.6 LiveLink take metadata format"), *CurrentDirectory);
+		TOptional<FTakeMetadata> ParseOldMetadataResult = LiveLinkMetadata::ParseOldLiveLinkTakeMetadata(CurrentDirectory, ValidationFailures);
+
+		if (ParseOldMetadataResult.IsSet())
+		{
+			return ParseOldMetadataResult.GetValue();
+		}
+
+		UE_LOG(LogGenericLiveLinkArchiveDevice, Display, TEXT("Checking directory (%s) for pre UE 5.6 StereoCamera take metadata format"), *CurrentDirectory);
+		ParseOldMetadataResult = StereoCameraMetadata::ParseOldStereoCameraMetadata(CurrentDirectory, ValidationFailures);
+
+		if (ParseOldMetadataResult.IsSet())
+		{
+			return ParseOldMetadataResult.GetValue();
+		}
+
+		UE_LOG(LogGenericLiveLinkArchiveDevice, Error, TEXT("Failed to parse take metadata file: %s"), *CurrentDirectory);
+
+		return {};
+	}
+
+	FTakeMetadata& TakeMetadata = TakeMetadataResult.GetValue();
+
+	TOptional<FString> ThumbnailPath = TakeMetadata.Thumbnail.GetThumbnailPath();
+	if (ThumbnailPath.IsSet())
+	{
+		if (FPaths::IsRelative(*ThumbnailPath))
+		{
+			TakeMetadata.Thumbnail = FTakeThumbnailData(FPaths::ConvertRelativePathToFull(CurrentDirectory, *ThumbnailPath));
+		}
+	}
+
+	if (!TakeMetadata.Video.IsEmpty())
+	{
+		for (FTakeMetadata::FVideo& Video : TakeMetadata.Video)
+		{
+			if (FPaths::IsRelative(*Video.Path))
+			{
+				Video.Path = FPaths::ConvertRelativePathToFull(CurrentDirectory, *Video.Path);
+			}
+		}
+	}
+
+	if (!TakeMetadata.Depth.IsEmpty())
+	{
+		for (FTakeMetadata::FVideo& Depth : TakeMetadata.Depth)
+		{
+			if (FPaths::IsRelative(*Depth.Path))
+			{
+				Depth.Path = FPaths::ConvertRelativePathToFull(CurrentDirectory, *Depth.Path);
+			}
+		}
+	}
+
+	if (!TakeMetadata.Audio.IsEmpty())
+	{
+		for (FTakeMetadata::FAudio& Audio : TakeMetadata.Audio)
+		{
+			if (FPaths::IsRelative(*Audio.Path))
+			{
+				Audio.Path = FPaths::ConvertRelativePathToFull(CurrentDirectory, *Audio.Path);
+			}
+		}
+	}
+
+	if (!TakeMetadata.Calibration.IsEmpty())
+	{
+		for (FTakeMetadata::FCalibration& Calib : TakeMetadata.Calibration)
+		{
+			if (FPaths::IsRelative(*Calib.Path))
+			{
+				Calib.Path = FPaths::ConvertRelativePathToFull(CurrentDirectory, *Calib.Path);
+			}
+		}
+	}
+
+	return TakeMetadataResult.StealValue();
+}
\ No newline at end of file
diff --git a/Engine/Plugins/VirtualProduction/CaptureManager/CaptureManagerDevices/Source/FilesystemLiveLinkDevice/Private/FilesystemLiveLinkDevice.h b/Engine/Plugins/VirtualProduction/CaptureManager/CaptureManagerDevices/Source/FilesystemLiveLinkDevice/Private/FilesystemLiveLinkDevice.h
new file mode 100644
--- /dev/null
+++ b/Engine/Plugins/VirtualProduction/CaptureManager/CaptureManagerDevices/Source/FilesystemLiveLinkDevice/Private/FilesystemLiveLinkDevice.h
@@ -0,0 +1,71 @@
+// Copyright Epic Games, Inc. All Rights Reserved.
+
+#pragma once
+
+#include "LiveLinkDevice.h"
+
+#include "Ingest/LiveLinkDeviceCapability_Ingest.h"
+#include "LiveLinkDeviceCapability_Connection.h"
+
+#include "Misc/Optional.h"
+
+#include "Async/TaskProgress.h"
+
+#include "BaseIngestLiveLinkDevice.h"
+
+#include "FilesystemLiveLinkDevice.generated.h"
+
+
+UCLASS(BlueprintType)
+class UFilesystemLiveLinkDeviceSettings : public ULiveLinkDeviceSettings
+{
+	GENERATED_BODY()
+
+public:
+	UPROPERTY(EditAnywhere, BlueprintReadWrite, Category = "Filesystem Device")
+	FString DisplayName = TEXT("Filesystem Device");
+
+	UPROPERTY(EditAnywhere, BlueprintReadWrite, Category = "Filesystem Device")
+	FDirectoryPath StoragePath;
+};
+
+
+UCLASS(BlueprintType)
+class UFilesystemLiveLinkDevice final : public UBaseIngestLiveLinkDevice
+	, public ILiveLinkDeviceCapability_Connection
+{
+	GENERATED_BODY()
+
+public:
+	UFUNCTION(BlueprintCallable, Category = "FilesystemDeviceSettings")
+	const UFilesystemLiveLinkDeviceSettings* GetSettings() const;
+
+	//~ Begin ULiveLinkDevice interface
+	virtual TSubclassOf<ULiveLinkDeviceSettings> StaticGetSettingsClass() const override;
+	virtual FText GetDisplayName() const override;
+	virtual EDeviceHealth GetDeviceHealth() const override;
+	virtual FText GetHealthText() const override;
+
+	virtual void OnSettingChanged(const FPropertyChangedEvent& InPropertyChangedEvent) override;
+	//~ End ULiveLinkDevice interface
+
+private:
+	virtual FString GetFullTakePath(UE::CaptureManager::FTakeId InTakeId) const override;
+
+	//~ Begin ULiveLinkDeviceCapability_Ingest interface
+	virtual void RunIngestTake(int32 InTakeId, const UIngestCapability_Options* InIngestOptions) override;
+	virtual void UpdateTakeList_Implementation() override;
+	//~ End ULiveLinkDeviceCapability_Ingest interface
+
+	//~ Begin ILiveLinkDeviceCapability_Connection interface
+	virtual ELiveLinkDeviceConnectionStatus GetConnectionStatus_Implementation() const override;
+	virtual FString GetHardwareId_Implementation() const override;
+	virtual bool SetHardwareId_Implementation(const FString& HardwareID) override;
+	virtual bool Connect_Implementation() override;
+	virtual bool Disconnect_Implementation() override;
+	//~ End ILiveLinkDeviceCapability_Connection interface
+
+	static TOptional<FTakeMetadata> ReadTake(const FString& InCurrentTakeFile);
+
+	TMap<int32, FString> RelativeTakePaths;
+};
diff --git a/Engine/Plugins/VirtualProduction/CaptureManager/CaptureManagerDevices/Source/FilesystemLiveLinkDevice/Private/FilesystemLiveLinkDeviceModule.cpp b/Engine/Plugins/VirtualProduction/CaptureManager/CaptureManagerDevices/Source/FilesystemLiveLinkDevice/Private/FilesystemLiveLinkDeviceModule.cpp
new file mode 100644
--- /dev/null
+++ b/Engine/Plugins/VirtualProduction/CaptureManager/CaptureManagerDevices/Source/FilesystemLiveLinkDevice/Private/FilesystemLiveLinkDeviceModule.cpp
@@ -0,0 +1,5 @@
+// Copyright Epic Games, Inc. All Rights Reserved.
+
+#include "Modules/ModuleManager.h"
+
+IMPLEMENT_MODULE(FDefaultModuleImpl, FilesystemLiveLinkDevice)
\ No newline at end of file
diff --git a/Engine/Plugins/VirtualProduction/CaptureManager/CaptureManagerDevices/Source/FilesystemLiveLinkDevice/Private/Tests/GenericLiveLinkDeviceTests.cpp b/Engine/Plugins/VirtualProduction/CaptureManager/CaptureManagerDevices/Source/FilesystemLiveLinkDevice/Private/Tests/GenericLiveLinkDeviceTests.cpp
new file mode 100644
--- /dev/null
+++ b/Engine/Plugins/VirtualProduction/CaptureManager/CaptureManagerDevices/Source/FilesystemLiveLinkDevice/Private/Tests/GenericLiveLinkDeviceTests.cpp
@@ -0,0 +1,38 @@
+// Copyright Epic Games, Inc. All Rights Reserved.
+
+#include "FilesystemLiveLinkDevice.h"
+
+#include "Misc/AutomationTest.h"
+
+IMPLEMENT_SIMPLE_AUTOMATION_TEST(FGenericLiveLinkArchiveDeviceCreationTest, "CaptureManager.FilesystemLiveLinkDevice.Creation",
+								 EAutomationTestFlags::EditorContext | EAutomationTestFlags::EngineFilter);
+
+bool FGenericLiveLinkArchiveDeviceCreationTest::RunTest(const FString& Parameters)
+{
+	UFilesystemLiveLinkDevice* Device = NewObject<UFilesystemLiveLinkDevice>();
+
+	TestNotNull(TEXT("Device Not Null"), Device);
+	TestTrue(TEXT("Device Implements Ingest Capability"), Device->Implements<ULiveLinkDeviceCapability_Ingest>());
+
+	return true;
+}
+
+IMPLEMENT_SIMPLE_AUTOMATION_TEST(FGenericLiveLinkArchiveDeviceIngestTakeAsyncTest, "CaptureManager.FilesystemLiveLinkDevice.IngestTakeAsync",
+								 EAutomationTestFlags::EditorContext | EAutomationTestFlags::EngineFilter);
+
+bool FGenericLiveLinkArchiveDeviceIngestTakeAsyncTest::RunTest(const FString& Parameters)
+{
+	UFilesystemLiveLinkDevice* Device = NewObject<UFilesystemLiveLinkDevice>();
+
+	TestNotNull(TEXT("Device Not Null"), Device);
+	TestTrue(TEXT("Device Implements Ingest Capability"), Device->Implements<ULiveLinkDeviceCapability_Ingest>());
+
+	UIngestCapability_ProcessHandle* IngestProcessHandle =
+		ILiveLinkDeviceCapability_Ingest::Execute_PrepareIngestProcess(Device, 0);
+	TestNotNull(TEXT("Ingest Process Handle Not Null"), IngestProcessHandle);
+
+	static constexpr UIngestCapability_Options* Options = nullptr;
+	IngestProcessHandle->Run(Options);
+
+	return true;
+}
\ No newline at end of file
diff --git a/Engine/Plugins/VirtualProduction/CaptureManager/CaptureManagerDevices/Source/FilesystemLiveLinkDevice/Private/Utils/GenericLiveLinkDeviceLog.h b/Engine/Plugins/VirtualProduction/CaptureManager/CaptureManagerDevices/Source/FilesystemLiveLinkDevice/Private/Utils/GenericLiveLinkDeviceLog.h
new file mode 100644
--- /dev/null
+++ b/Engine/Plugins/VirtualProduction/CaptureManager/CaptureManagerDevices/Source/FilesystemLiveLinkDevice/Private/Utils/GenericLiveLinkDeviceLog.h
@@ -0,0 +1,7 @@
+// Copyright Epic Games, Inc. All Rights Reserved.
+
+#pragma once
+
+#include "Logging/LogMacros.h"
+
+DECLARE_LOG_CATEGORY_EXTERN(LogGenericLiveLinkArchiveDevice, Log, All);
