LiveLink Face App fails to ingest two takes at the same time

diff --git a/Engine/Plugins/VirtualProduction/CaptureManager/CaptureManagerDevices/Source/CPSLiveLinkDevice/Private/LiveLinkFaceDevice.cpp b/Engine/Plugins/VirtualProduction/CaptureManager/CaptureManagerDevices/Source/CPSLiveLinkDevice/Private/LiveLinkFaceDevice.cpp
--- a/Engine/Plugins/VirtualProduction/CaptureManager/CaptureManagerDevices/Source/CPSLiveLinkDevice/Private/LiveLinkFaceDevice.cpp
+++ b/Engine/Plugins/VirtualProduction/CaptureManager/CaptureManagerDevices/Source/CPSLiveLinkDevice/Private/LiveLinkFaceDevice.cpp
@@ -366,10 +366,10 @@ void ULiveLinkFaceDevice::OnExportProgressReport(float InProgress,
 }
 
 void ULiveLinkFaceDevice::OnExportFinished(UE::CaptureManager::TProtocolResult<void> InResult,
-												 FString InTakeName,
-												 UE::CaptureManager::FTakeId InTakeId,
-												 TSharedPtr<UE::CaptureManager::FTaskProgress> InTaskProgress,
-												 TStrongObjectPtr<const UIngestCapability_Options> InIngestOptions)
+										   FString InTakeName,
+										   UE::CaptureManager::FTakeId InTakeId,
+										   TSharedPtr<UE::CaptureManager::FTaskProgress> InTaskProgress,
+										   TStrongObjectPtr<const UIngestCapability_Options> InIngestOptions)
 {
 	using namespace UE::CaptureManager;
 
@@ -387,22 +387,35 @@ void ULiveLinkFaceDevice::OnExportFinished(UE::CaptureManager::TProtocolResult<v
 			Lock.Unlock();
 
 			UpdateTake(InTakeId, MoveTemp(TakeMetadata.GetValue()));
-			Super::IngestTake(InTakeId, InIngestOptions.Get(), InTaskProgress);
+
+			// Free the current thread that is waiting on next download
+			TWeakObjectPtr<ULiveLinkFaceDevice> WeakThis(this);
+			AsyncTask(ENamedThreads::Type::AnyThread, [WeakThis = MoveTemp(WeakThis), TakeId = InTakeId, IngestOptions = MoveTemp(InIngestOptions), TaskProgress = MoveTemp(InTaskProgress)]()
+			{
+				TStrongObjectPtr<ULiveLinkFaceDevice> StrongThis = WeakThis.Pin();
+
+				if (!StrongThis.IsValid())
+				{
+					return;
+				}
+
+				StrongThis->Super::IngestTake(TakeId, IngestOptions.Get(), TaskProgress);
+				StrongThis->RemoveDownloadedTakeData(TakeId);
+			});	
 		}
 		else
 		{
 			TValueOrError<void, FIngestCapability_Error> Result = MakeError(FIngestCapability_Error::DownloaderError, TEXT("Failed to parse the take metadata"));
 			ExecuteProcessFinishedReporter(InTakeId, MoveTemp(Result));
+			IFileManager::Get().DeleteDirectory(*DownloadedTake, false, true);
 		}
 	}
 	else
 	{
 		TValueOrError<void, FIngestCapability_Error> Result = MakeError(FIngestCapability_Error::DownloaderError, InResult.GetError().GetMessage());
-
 		ExecuteProcessFinishedReporter(InTakeId, MoveTemp(Result));
+		IFileManager::Get().DeleteDirectory(*DownloadedTake, false, true);
 	}
-
-	IFileManager::Get().DeleteDirectory(*DownloadedTake, false, true);
 }
 
 TOptional<FTakeMetadata> ULiveLinkFaceDevice::ParseTake(const FString& InTakeDirectory, const FString& InTakeName)
@@ -487,3 +500,16 @@ void ULiveLinkFaceDevice::FetchThumbnails(TMap<FString, UE::CaptureManager::FTak
 
 	Device->FetchThumbnails(MoveTemp(DataStream));
 }
+
+void ULiveLinkFaceDevice::RemoveDownloadedTakeData(const UE::CaptureManager::FTakeId InTakeId)
+{
+	FScopeLock Lock(&DownloadedTakesMutex);
+	FString* DownloadedTake = DownloadedTakes.Find(InTakeId);
+
+	if (!DownloadedTake)
+	{
+		return;
+	}
+
+	IFileManager::Get().DeleteDirectory(**DownloadedTake, false, true);
+}
\ No newline at end of file
diff --git a/Engine/Plugins/VirtualProduction/CaptureManager/CaptureManagerDevices/Source/CPSLiveLinkDevice/Private/LiveLinkFaceDevice.h b/Engine/Plugins/VirtualProduction/CaptureManager/CaptureManagerDevices/Source/CPSLiveLinkDevice/Private/LiveLinkFaceDevice.h
--- a/Engine/Plugins/VirtualProduction/CaptureManager/CaptureManagerDevices/Source/CPSLiveLinkDevice/Private/LiveLinkFaceDevice.h
+++ b/Engine/Plugins/VirtualProduction/CaptureManager/CaptureManagerDevices/Source/CPSLiveLinkDevice/Private/LiveLinkFaceDevice.h
@@ -102,6 +102,8 @@ private:
 	void FetchThumbnail(const FString& InTakeName, UE::CaptureManager::FTakeId InTakeId);
 	void FetchThumbnails(TMap<FString, UE::CaptureManager::FTakeId> InNameToIdMap);
 
+	void RemoveDownloadedTakeData(const UE::CaptureManager::FTakeId InTakeId);
+
 	TSharedPtr<UE::CaptureManager::FCPSDevice> Device;
 
 	mutable FCriticalSection DownloadedTakesMutex;
diff --git a/Engine/Plugins/VirtualProduction/CaptureManager/CaptureManagerDevices/Source/CPSLiveLinkDevice/Private/Protocol/CPSDevice.cpp b/Engine/Plugins/VirtualProduction/CaptureManager/CaptureManagerDevices/Source/CPSLiveLinkDevice/Private/Protocol/CPSDevice.cpp
--- a/Engine/Plugins/VirtualProduction/CaptureManager/CaptureManagerDevices/Source/CPSLiveLinkDevice/Private/Protocol/CPSDevice.cpp
+++ b/Engine/Plugins/VirtualProduction/CaptureManager/CaptureManagerDevices/Source/CPSLiveLinkDevice/Private/Protocol/CPSDevice.cpp
@@ -414,7 +414,7 @@ void FCPSDevice::OnConnectTick()
 {
 	if (!bIsConnected.load())
 	{
-		UE_LOG(LogCPSProtocolDevice, Display, TEXT("Connecting to the CPS Device: %s:%d"), *DeviceIpAddress, DeviceControlPort);
+		UE_LOG(LogCPSProtocolDevice, Log, TEXT("Connecting to the CPS Device: %s:%d"), *DeviceIpAddress, DeviceControlPort);
 		ConnThread->Add(FEmpty()); // Initiate connect
 	}
 	else
