[CaptureManager] Live Link Devices have friendly display names

diff --git a/Engine/Plugins/VirtualProduction/CaptureManager/CaptureManagerDevices/Source/CPSLiveLinkDevice/Private/LiveLinkFaceDevice.h b/Engine/Plugins/VirtualProduction/CaptureManager/CaptureManagerDevices/Source/CPSLiveLinkDevice/Private/LiveLinkFaceDevice.h
--- a/Engine/Plugins/VirtualProduction/CaptureManager/CaptureManagerDevices/Source/CPSLiveLinkDevice/Private/LiveLinkFaceDevice.h
+++ b/Engine/Plugins/VirtualProduction/CaptureManager/CaptureManagerDevices/Source/CPSLiveLinkDevice/Private/LiveLinkFaceDevice.h
@@ -40,7 +40,7 @@ public:
 };
 
 
-UCLASS(BlueprintType)
+UCLASS(BlueprintType, meta = (DisplayName = "Live Link Face"))
 class ULiveLinkFaceDevice final : public UBaseIngestLiveLinkDevice
 	, public ILiveLinkDeviceCapability_Connection
 	, public ILiveLinkDeviceCapability_Recording
diff --git a/Engine/Plugins/VirtualProduction/CaptureManager/CaptureManagerDevices/Source/FilesystemLiveLinkDevice/Private/FilesystemLiveLinkDevice.h b/Engine/Plugins/VirtualProduction/CaptureManager/CaptureManagerDevices/Source/FilesystemLiveLinkDevice/Private/FilesystemLiveLinkDevice.h
--- a/Engine/Plugins/VirtualProduction/CaptureManager/CaptureManagerDevices/Source/FilesystemLiveLinkDevice/Private/FilesystemLiveLinkDevice.h
+++ b/Engine/Plugins/VirtualProduction/CaptureManager/CaptureManagerDevices/Source/FilesystemLiveLinkDevice/Private/FilesystemLiveLinkDevice.h
@@ -30,7 +30,7 @@ public:
 };
 
 
-UCLASS(BlueprintType)
+UCLASS(BlueprintType, meta=(DisplayName="Filesystem"))
 class UFilesystemLiveLinkDevice final : public UBaseIngestLiveLinkDevice
 	, public ILiveLinkDeviceCapability_Connection
 {
diff --git a/Engine/Plugins/VirtualProduction/CaptureManager/CaptureManagerDevices/Source/MonoVideoLiveLinkDevice/Private/MonoVideoLiveLinkDevice.h b/Engine/Plugins/VirtualProduction/CaptureManager/CaptureManagerDevices/Source/MonoVideoLiveLinkDevice/Private/MonoVideoLiveLinkDevice.h
--- a/Engine/Plugins/VirtualProduction/CaptureManager/CaptureManagerDevices/Source/MonoVideoLiveLinkDevice/Private/MonoVideoLiveLinkDevice.h
+++ b/Engine/Plugins/VirtualProduction/CaptureManager/CaptureManagerDevices/Source/MonoVideoLiveLinkDevice/Private/MonoVideoLiveLinkDevice.h
@@ -45,7 +45,7 @@ public:
 };
 
 
-UCLASS(BlueprintType)
+UCLASS(BlueprintType, meta = (DisplayName = "Mono Video"))
 class UMonoVideoLiveLinkDevice final : public UBaseIngestLiveLinkDevice
 	, public ILiveLinkDeviceCapability_Connection
 {
