[CaptureManager] Fix crash in Live Link Face device

diff --git a/Engine/Plugins/VirtualProduction/CaptureManager/CaptureManagerDevices/Source/CPSLiveLinkDevice/Private/LiveLinkFaceDevice.cpp b/Engine/Plugins/VirtualProduction/CaptureManager/CaptureManagerDevices/Source/CPSLiveLinkDevice/Private/LiveLinkFaceDevice.cpp
--- a/Engine/Plugins/VirtualProduction/CaptureManager/CaptureManagerDevices/Source/CPSLiveLinkDevice/Private/LiveLinkFaceDevice.cpp
+++ b/Engine/Plugins/VirtualProduction/CaptureManager/CaptureManagerDevices/Source/CPSLiveLinkDevice/Private/LiveLinkFaceDevice.cpp
@@ -45,6 +45,8 @@ FText ULiveLinkFaceDevice::GetHealthText() const
 void ULiveLinkFaceDevice::OnDeviceAdded()
 {
 	GetDeviceSettings<ULiveLinkFaceDeviceSettings>()->ConnectAction.DeviceGuid = GetDeviceId();
+
+	Super::OnDeviceAdded();
 }
 
 void ULiveLinkFaceDevice::OnDeviceRemoved()
