Implement simple stereo video device for Capture Manager
Add thumbnail support for mono and stereo video devices
Rename {UserId} discovery token to {Name} for video and autio components

diff --git a/Engine/Plugins/VirtualProduction/CaptureManager/CaptureManagerApp/Source/CaptureManagerMediaRW/CaptureManagerMediaRW.Build.cs b/Engine/Plugins/VirtualProduction/CaptureManager/CaptureManagerApp/Source/CaptureManagerMediaRW/CaptureManagerMediaRW.Build.cs
--- a/Engine/Plugins/VirtualProduction/CaptureManager/CaptureManagerApp/Source/CaptureManagerMediaRW/CaptureManagerMediaRW.Build.cs
+++ b/Engine/Plugins/VirtualProduction/CaptureManager/CaptureManagerApp/Source/CaptureManagerMediaRW/CaptureManagerMediaRW.Build.cs
@@ -10,13 +10,13 @@ public class CaptureManagerMediaRW : ModuleRules
 
 		PublicDependencyModuleNames.AddRange(new string[]
 		{
-			"Core"
+			"Core",
+			"Media"
 		});
 
 		PrivateDependencyModuleNames.AddRange(new string[]
 		{
 			"Engine",
-			"Media",
 			"RenderCore",
 			"ImageWriteQueue",
 			"ImageCore",
diff --git a/Engine/Plugins/VirtualProduction/CaptureManager/CaptureManagerCore/Source/CaptureManagerTakeMetadata/Public/CaptureManagerTakeMetadata.h b/Engine/Plugins/VirtualProduction/CaptureManager/CaptureManagerCore/Source/CaptureManagerTakeMetadata/Public/CaptureManagerTakeMetadata.h
--- a/Engine/Plugins/VirtualProduction/CaptureManager/CaptureManagerCore/Source/CaptureManagerTakeMetadata/Public/CaptureManagerTakeMetadata.h
+++ b/Engine/Plugins/VirtualProduction/CaptureManager/CaptureManagerCore/Source/CaptureManagerTakeMetadata/Public/CaptureManagerTakeMetadata.h
@@ -19,7 +19,7 @@ public:
 	/** Construct thumbnail from file path. */
 	FTakeThumbnailData(FString InThumbnailPath);
 
-	/** Construct thumbnail from data array. */
+	/** Construct thumbnail from file data array. */
 	FTakeThumbnailData(TArray<uint8> InThumbnailData);
 
 	/** Get thumbnail. */
diff --git a/Engine/Plugins/VirtualProduction/CaptureManager/CaptureManagerDevices/CaptureManagerDevices.uplugin b/Engine/Plugins/VirtualProduction/CaptureManager/CaptureManagerDevices/CaptureManagerDevices.uplugin
--- a/Engine/Plugins/VirtualProduction/CaptureManager/CaptureManagerDevices/CaptureManagerDevices.uplugin
+++ b/Engine/Plugins/VirtualProduction/CaptureManager/CaptureManagerDevices/CaptureManagerDevices.uplugin
@@ -31,6 +31,16 @@
 			"Name": "FilesystemLiveLinkDevice",
 			"Type": "Editor",
 			"LoadingPhase": "Default"
+		},
+		{
+			"Name": "StereoVideoLiveLinkDevice",
+			"Type": "Editor",
+			"LoadingPhase": "Default"
+		},
+		{
+			"Name": "VideoLiveLinkDeviceCommon",
+			"Type": "Editor",
+			"LoadingPhase": "Default"
 		}
 	],
 	"Plugins": [
diff --git a/Engine/Plugins/VirtualProduction/CaptureManager/CaptureManagerDevices/Source/MonoVideoLiveLinkDevice/MonoVideoLiveLinkDevice.Build.cs b/Engine/Plugins/VirtualProduction/CaptureManager/CaptureManagerDevices/Source/MonoVideoLiveLinkDevice/MonoVideoLiveLinkDevice.Build.cs
--- a/Engine/Plugins/VirtualProduction/CaptureManager/CaptureManagerDevices/Source/MonoVideoLiveLinkDevice/MonoVideoLiveLinkDevice.Build.cs
+++ b/Engine/Plugins/VirtualProduction/CaptureManager/CaptureManagerDevices/Source/MonoVideoLiveLinkDevice/MonoVideoLiveLinkDevice.Build.cs
@@ -33,7 +33,8 @@ public class MonoVideoLiveLinkDevice : ModuleRules
 			"CaptureManagerTakeMetadata",
 			"CaptureManagerSettings",
 			"CoreUObject",
-			"UnrealEd"
+			"UnrealEd",
+			"VideoLiveLinkDeviceCommon"
 		});
 
 		CppCompileWarningSettings.UnsafeTypeCastWarningLevel = WarningLevel.Error;
diff --git a/Engine/Plugins/VirtualProduction/CaptureManager/CaptureManagerDevices/Source/MonoVideoLiveLinkDevice/Private/MonoVideoLiveLinkDevice.cpp b/Engine/Plugins/VirtualProduction/CaptureManager/CaptureManagerDevices/Source/MonoVideoLiveLinkDevice/Private/MonoVideoLiveLinkDevice.cpp
--- a/Engine/Plugins/VirtualProduction/CaptureManager/CaptureManagerDevices/Source/MonoVideoLiveLinkDevice/Private/MonoVideoLiveLinkDevice.cpp
+++ b/Engine/Plugins/VirtualProduction/CaptureManager/CaptureManagerDevices/Source/MonoVideoLiveLinkDevice/Private/MonoVideoLiveLinkDevice.cpp
@@ -4,17 +4,27 @@
 
 #include "HAL/FileManager.h"
 #include "Settings/CaptureManagerSettings.h"
+#include "ImageUtils.h"
 
 #include "CaptureManagerMediaRWModule.h"
 
-#include "Utils/MonoVideoDeviceTokenParser.h"
+#include "Utils/VideoDeviceTokenParser.h"
+#include "Utils/VideoDeviceThumbnailExtractor.h"
 
-#include "MonoVideoLiveLinkDeviceLog.h"
+#include "VideoLiveLinkDeviceLog.h"
 
-DEFINE_LOG_CATEGORY(LogMonoVideoLiveLinkDevice);
+DEFINE_LOG_CATEGORY(LogVideoLiveLinkDevice);
 
 namespace UE::MonoVideoLiveLinkDevice::Private
 {
+
+static const TArray<FString::ElementType> Delimiters =
+{
+	'-',
+	'_',
+	'.'
+};
+
 static const TArray<FStringView> SupportedVideoFormats =
 {
 	TEXTVIEW("mp4"),
@@ -195,23 +205,23 @@ TOptional<FTakeMetadata> UMonoVideoLiveLinkDevice::ReadTake(const FString& InCur
 	FString FileName = FPaths::GetBaseFilename(InCurrentTakeFile);
 
 	FString SlateName;
-	FString UserId;
+	FString Name;
 	int32 TakeNumber = INDEX_NONE;
 
 	if (!FileNameFormat.IsEmpty())
 	{
-		FMonoVideoDeviceTokenParser TokenParser(FileNameFormat, FileName);
+		FVideoDeviceTokenParser TokenParser(FileNameFormat, FileName, UE::MonoVideoLiveLinkDevice::Private::Delimiters);
 		bool bParse = TokenParser.Parse();
 
 		if (!bParse)
 		{
-			UE_LOG(LogMonoVideoLiveLinkDevice, Warning, TEXT("Failed to match the specified format (%s) with the video file (%s)"), *(FileNameFormat), *InCurrentTakeFile);
+			UE_LOG(LogVideoLiveLinkDevice, Warning, TEXT("Failed to match the specified format (%s) with the video file (%s)"), *(FileNameFormat), *InCurrentTakeFile);
 			return {};
 		}
 
 		SlateName = TokenParser.GetSlateName();
 		TakeNumber = TokenParser.GetTakeNumber();
-		UserId = TokenParser.GetUserId();
+		Name = TokenParser.GetName();
 	}
 
 	if (SlateName.IsEmpty())
@@ -224,9 +234,9 @@ TOptional<FTakeMetadata> UMonoVideoLiveLinkDevice::ReadTake(const FString& InCur
 		TakeNumber = 1;
 	}
 
-	if (UserId.IsEmpty())
+	if (Name.IsEmpty())
 	{
-		UserId = TEXT("video");
+		Name = TEXT("video");
 	}
 
 	TakeMetadata.Version.Major = 4;
@@ -235,7 +245,7 @@ TOptional<FTakeMetadata> UMonoVideoLiveLinkDevice::ReadTake(const FString& InCur
 	TakeMetadata.Slate = MoveTemp(SlateName);
 	TakeMetadata.TakeNumber = TakeNumber;
 
-	// Create new guid for old takes
+	// Create new guid
 	TakeMetadata.UniqueId = FGuid::NewGuid().ToString(EGuidFormats::DigitsWithHyphensLower);
 	TakeMetadata.DateTime = FileData.CreationTime;
 	TakeMetadata.Device.Model = TEXT("MonoVideo");
@@ -244,14 +254,31 @@ TOptional<FTakeMetadata> UMonoVideoLiveLinkDevice::ReadTake(const FString& InCur
 	TValueOrError<TUniquePtr<IVideoReader>, FText> VideoReaderResult = MediaManager.CreateVideoReader(InCurrentTakeFile);
 	if (VideoReaderResult.HasValue())
 	{
-		Video.FrameRate = static_cast<float>(VideoReaderResult.StealValue()->GetFrameRate().AsDecimal());
+		TUniquePtr<IVideoReader> VideoReader = VideoReaderResult.StealValue();
+		Video.FrameRate = static_cast<float>(VideoReader->GetFrameRate().AsDecimal());
+		TValueOrError<TUniquePtr<UE::CaptureManager::FMediaTextureSample>, FText> TextureSampleResult = VideoReader->Next();
+		if (TextureSampleResult.HasValue())
+		{
+			FVideoDeviceThumbnailExtractor ThumbnailExtractor;
+			UE::CaptureManager::FMediaTextureSample* Sample = TextureSampleResult.GetValue().Get();
+
+			TArray<FColor> ThumbnailRawColorData = ThumbnailExtractor.ExtractThumbnail(Sample);
+
+			if (ThumbnailRawColorData.Num() > 0)
+			{
+				TArray<uint8> ThumbnailData;
+				FImageUtils::ThumbnailCompressImageArray(Sample->Dimensions.X, Sample->Dimensions.Y, ThumbnailRawColorData, ThumbnailData);
+
+				TakeMetadata.Thumbnail = ThumbnailData;
+			}
+		}
 	}
 	else
 	{
-		UE_LOG(LogMonoVideoLiveLinkDevice, Warning, TEXT("Failed to determine the frame rate for the video file %s"), *InCurrentTakeFile);
+		UE_LOG(LogVideoLiveLinkDevice, Warning, TEXT("Failed to determine the frame rate for the video file %s"), *InCurrentTakeFile);
 	}
 
-	Video.Name = MoveTemp(UserId);
+	Video.Name = MoveTemp(Name);
 	Video.Format = FPaths::GetExtension(InCurrentTakeFile);
 	Video.Path = InCurrentTakeFile;
 	Video.PathType = FTakeMetadata::FVideo::EPathType::File;
diff --git a/Engine/Plugins/VirtualProduction/CaptureManager/CaptureManagerDevices/Source/MonoVideoLiveLinkDevice/Private/MonoVideoLiveLinkDevice.h b/Engine/Plugins/VirtualProduction/CaptureManager/CaptureManagerDevices/Source/MonoVideoLiveLinkDevice/Private/MonoVideoLiveLinkDevice.h
--- a/Engine/Plugins/VirtualProduction/CaptureManager/CaptureManagerDevices/Source/MonoVideoLiveLinkDevice/Private/MonoVideoLiveLinkDevice.h
+++ b/Engine/Plugins/VirtualProduction/CaptureManager/CaptureManagerDevices/Source/MonoVideoLiveLinkDevice/Private/MonoVideoLiveLinkDevice.h
@@ -22,10 +22,10 @@ class UMonoVideoLiveLinkDeviceSettings : public ULiveLinkDeviceSettings
 	GENERATED_BODY()
 
 public:
-	UPROPERTY(EditAnywhere, BlueprintReadWrite, Category = "Mono Video Live Link Device")
-	FString DisplayName = TEXT("Mono Video Live Link Device");
+	UPROPERTY(EditAnywhere, BlueprintReadWrite, Category = "Mono Video Device")
+	FString DisplayName = TEXT("Mono Video Device");
 
-	UPROPERTY(EditAnywhere, BlueprintReadWrite, Category = "Mono Video Live Link Device")
+	UPROPERTY(EditAnywhere, BlueprintReadWrite, Category = "Mono Video Device")
 	FDirectoryPath StoragePath;
 
 	/** 
@@ -33,18 +33,18 @@ public:
 	 *		Available tokens:
 	 *			* {SlateName}
 	 *			* {TakeNumber}
-	 *			* {UserId}
+	 *			* {Name}
 	 *			* {Any}
 	 *		Available delimiters: '_', '-', '.' 
 	 * 
 	 * Leave empty to parse every video file in the specified directory
 	 */
-	UPROPERTY(EditAnywhere, BlueprintReadWrite, Category = "Mono Video Live Link Device")
+	UPROPERTY(EditAnywhere, BlueprintReadWrite, Category = "Mono Video Device")
 	FString FileNameFormat = TEXT("{SlateName}_{TakeNumber}_{Any}");
 };
 
 
-UCLASS(BlueprintType, meta = (DisplayName = "Mono Video"))
+UCLASS(BlueprintType, meta = (DisplayName = "Mono Video Device"))
 class UMonoVideoLiveLinkDevice final : public UBaseIngestLiveLinkDevice
 	, public ILiveLinkDeviceCapability_Connection
 {
diff --git a/Engine/Plugins/VirtualProduction/CaptureManager/CaptureManagerDevices/Source/MonoVideoLiveLinkDevice/Private/MonoVideoLiveLinkDeviceLog.h b/Engine/Plugins/VirtualProduction/CaptureManager/CaptureManagerDevices/Source/MonoVideoLiveLinkDevice/Private/MonoVideoLiveLinkDeviceLog.h
deleted file mode 100644
diff --git a/Engine/Plugins/VirtualProduction/CaptureManager/CaptureManagerDevices/Source/MonoVideoLiveLinkDevice/Private/Tests/TestMonoVideoDeviceTokenParser.cpp b/Engine/Plugins/VirtualProduction/CaptureManager/CaptureManagerDevices/Source/MonoVideoLiveLinkDevice/Private/Tests/TestMonoVideoDeviceTokenParser.cpp
deleted file mode 100644
diff --git a/Engine/Plugins/VirtualProduction/CaptureManager/CaptureManagerDevices/Source/MonoVideoLiveLinkDevice/Private/Utils/MonoVideoDeviceTokenParser.cpp b/Engine/Plugins/VirtualProduction/CaptureManager/CaptureManagerDevices/Source/MonoVideoLiveLinkDevice/Private/Utils/MonoVideoDeviceTokenParser.cpp
deleted file mode 100644
diff --git a/Engine/Plugins/VirtualProduction/CaptureManager/CaptureManagerDevices/Source/MonoVideoLiveLinkDevice/Private/Utils/MonoVideoDeviceTokenParser.h b/Engine/Plugins/VirtualProduction/CaptureManager/CaptureManagerDevices/Source/MonoVideoLiveLinkDevice/Private/Utils/MonoVideoDeviceTokenParser.h
deleted file mode 100644
diff --git a/Engine/Plugins/VirtualProduction/CaptureManager/CaptureManagerDevices/Source/StereoVideoLiveLinkDevice/Private/StereoVideoLiveLinkDevice.cpp b/Engine/Plugins/VirtualProduction/CaptureManager/CaptureManagerDevices/Source/StereoVideoLiveLinkDevice/Private/StereoVideoLiveLinkDevice.cpp
new file mode 100644
--- /dev/null
+++ b/Engine/Plugins/VirtualProduction/CaptureManager/CaptureManagerDevices/Source/StereoVideoLiveLinkDevice/Private/StereoVideoLiveLinkDevice.cpp
@@ -0,0 +1,563 @@
+	// Copyright Epic Games, Inc. All Rights Reserved.
+
+#include "StereoVideoLiveLinkDevice.h"
+
+#include "HAL/FileManager.h"
+#include "Settings/CaptureManagerSettings.h"
+#include "ImageUtils.h"
+#include "CaptureManagerMediaRWModule.h"
+#include "Asset/CaptureAssetSanitization.h"
+#include "Misc/FileHelper.h"
+#include "Engine/Texture2D.h"
+
+#include "Utils/VideoDeviceTokenParser.h"
+#include "Utils/VideoDeviceThumbnailExtractor.h"
+
+#include "VideoLiveLinkDeviceLog.h"
+
+DEFINE_LOG_CATEGORY(LogVideoLiveLinkDevice);
+#define LOCTEXT_NAMESPACE "StereoVideoLLDevice"
+
+namespace UE::StereoVideoLiveLinkDevice::Private
+{
+	static const TArray<FString::ElementType> Delimiters =
+	{
+		'-',
+		'_',
+		'.',
+		'/'
+	};
+
+	static const TArray<FStringView> SupportedVideoExtensions =
+	{
+		TEXTVIEW("mp4"),
+		TEXTVIEW("mov")
+	};
+
+	static const TArray<FStringView> SupportedAudioExtensions =
+	{
+		TEXTVIEW("wav")
+	};
+
+	static const TArray<FStringView> SupportedImageSequenceExtensions =
+	{
+		TEXTVIEW("jpg"),
+		TEXTVIEW("jpeg"),
+		TEXTVIEW("png")
+	};
+
+	FMediaRWManager& GetMediaRWManager()
+	{
+		return FModuleManager::LoadModuleChecked<FCaptureManagerMediaRWModule>("CaptureManagerMediaRW").Get();
+	}
+}
+
+const UStereoVideoLiveLinkDeviceSettings* UStereoVideoLiveLinkDevice::GetSettings() const
+{
+	return GetDeviceSettings<UStereoVideoLiveLinkDeviceSettings>();
+}
+
+TSubclassOf<ULiveLinkDeviceSettings> UStereoVideoLiveLinkDevice::GetSettingsClass() const
+{
+	return UStereoVideoLiveLinkDeviceSettings::StaticClass();
+}
+
+FText UStereoVideoLiveLinkDevice::GetDisplayName() const
+{
+	return FText::FromString(GetSettings()->DisplayName);
+}
+
+EDeviceHealth UStereoVideoLiveLinkDevice::GetDeviceHealth() const
+{
+	return EDeviceHealth::Nominal;
+}
+
+FText UStereoVideoLiveLinkDevice::GetHealthText() const
+{
+	return FText::FromString("Nominal");
+}
+
+void UStereoVideoLiveLinkDevice::OnSettingChanged(const FPropertyChangedEvent& InPropertyChangedEvent)
+{
+	Super::OnSettingChanged(InPropertyChangedEvent);
+
+	FString Path = GetSettings()->StoragePath.Path;
+	if (!Path.IsEmpty() && FPaths::DirectoryExists(Path))
+	{
+		SetConnectionStatus(ELiveLinkDeviceConnectionStatus::Connected);
+	}
+	else
+	{
+		SetConnectionStatus(ELiveLinkDeviceConnectionStatus::Disconnected);
+	}
+
+}
+
+FString UStereoVideoLiveLinkDevice::GetFullTakePath(UE::CaptureManager::FTakeId InTakeId) const
+{
+	if (const FString* TakePath = FullTakePaths.Find(InTakeId))
+	{
+		return *TakePath;
+	}
+
+	return FString();
+}
+
+void UStereoVideoLiveLinkDevice::RunIngestTake(int32 InTakeId, const UIngestCapability_Options* InIngestOptions)
+{
+	AsyncTask(ENamedThreads::AnyThread, [this,
+		TakeId = InTakeId,
+		IngestOptions = InIngestOptions]()
+		{
+			static constexpr uint32 NumberOfTasks = 2;
+
+			using namespace UE::CaptureManager;
+			TSharedPtr<FTaskProgress> TaskProgress
+				= MakeShared<FTaskProgress>(NumberOfTasks, FTaskProgress::FProgressReporter::CreateLambda([this, TakeId](double InProgress)
+					{
+						ExecuteProcessProgressReporter(TakeId, InProgress);
+					}));
+
+			Super::IngestTake(TakeId, IngestOptions, MoveTemp(TaskProgress));
+		});
+}
+
+int32 UStereoVideoLiveLinkDevice::FTakeWithComponents::CountComponents(UStereoVideoLiveLinkDevice::ETakeComponentType Type)
+{
+	int32 Count = 0;
+	for (const UStereoVideoLiveLinkDevice::FTakeWithComponents::Component& Component : Components)
+	{
+		if (Component.Type == Type)
+		{
+			++Count;
+		}
+	}
+	return Count;
+}
+
+void UStereoVideoLiveLinkDevice::UpdateTakeList_Implementation()
+{
+	RemoveAllTakes();
+
+	FString StoragePath = GetSettings()->StoragePath.Path;
+
+	TMap<FString, UStereoVideoLiveLinkDevice::FTakeWithComponents> TakesWithComponentsCandidates = DiscoverTakes(StoragePath);
+
+	TArray<FString> SlateNames;
+	TakesWithComponentsCandidates.GetKeys(SlateNames);
+
+	for (const FString& SlateName : SlateNames)
+	{
+		FTakeWithComponents& TakeWithComponents = TakesWithComponentsCandidates[SlateName];
+
+		bool TakeAsExpected = ((TakeWithComponents.CountComponents(ETakeComponentType::VIDEO) == 2 && TakeWithComponents.CountComponents(ETakeComponentType::IMAGE_SEQUENCE) == 0) ||
+			(TakeWithComponents.CountComponents(ETakeComponentType::VIDEO) == 0 && TakeWithComponents.CountComponents(ETakeComponentType::IMAGE_SEQUENCE) == 2)) &&
+			TakeWithComponents.CountComponents(ETakeComponentType::AUDIO) <= 1;
+
+		if (!TakeAsExpected)
+		{
+			UE_LOG(LogVideoLiveLinkDevice, Warning, TEXT("Invalid take at '%s'. Take should have exactly two video files or two image sequences. Optionally, one audio file."), *TakeWithComponents.TakePath);
+			continue;
+		}
+
+		TOptional<FTakeMetadata> TakeInfoResult = CreateTakeMetadata(TakeWithComponents);
+
+		if (TakeInfoResult.IsSet())
+		{
+			FTakeMetadata TakeMetadata = TakeInfoResult.GetValue();
+
+			int32 TakeId = AddTake(MoveTemp(TakeMetadata));
+			FullTakePaths.Add(TakeId, TakeWithComponents.TakePath);
+		}
+	}
+
+	ExecuteUpdateTakeListCallback(Execute_GetTakeIdentifiers(this));
+}
+
+TMap<FString, UStereoVideoLiveLinkDevice::FTakeWithComponents> UStereoVideoLiveLinkDevice::DiscoverTakes(FString StoragePath)
+{
+	TMap<FString, FTakeWithComponents> TakesWithComponentsCandidates;
+
+	const bool bIterationResult = IFileManager::Get().IterateDirectoryRecursively(*StoragePath, [&](const TCHAR* InPath, bool bInIsDirectory)
+		{
+			FString Extension = FPaths::GetExtension(InPath);
+
+			if (UE::StereoVideoLiveLinkDevice::Private::SupportedVideoExtensions.Contains(Extension))
+			{
+				ExtractTakeComponents(InPath, StoragePath, ETakeComponentType::VIDEO, GetSettings()->VideoFileNameFormat, "UnknownVideoName", TakesWithComponentsCandidates);
+			}
+			else if (UE::StereoVideoLiveLinkDevice::Private::SupportedImageSequenceExtensions.Contains(Extension))
+			{
+				ExtractTakeComponents(InPath, StoragePath, ETakeComponentType::IMAGE_SEQUENCE, GetSettings()->VideoFileNameFormat, "UnknownImageSequenceName", TakesWithComponentsCandidates);
+			}
+			else if (UE::StereoVideoLiveLinkDevice::Private::SupportedAudioExtensions.Contains(Extension))
+			{
+				ExtractTakeComponents(InPath, StoragePath, ETakeComponentType::AUDIO, GetSettings()->AudioFileNameFormat, "UnknownAudioName", TakesWithComponentsCandidates);
+			}
+
+			return true;
+		});
+
+	return TakesWithComponentsCandidates;
+}
+
+void UStereoVideoLiveLinkDevice::ExtractTakeComponents(FString ComponentPath, FString StoragePath, ETakeComponentType ComponentType, FString Format, FString UnknownComponentName, TMap<FString, UStereoVideoLiveLinkDevice::FTakeWithComponents>& OutTakesWithComponentsCandidates)
+{
+	TOptional<FTakeWithComponents> TakeWithComponents;
+
+	if (Format == TEXT("{Any}"))
+	{
+		TakeWithComponents = ExtractTakeComponentsFromDirectoryStructure(ComponentPath, StoragePath, ComponentType);
+	}
+	else
+	{
+		TakeWithComponents = ExtractTakeComponentsUsingTokens(ComponentPath, StoragePath, Format, ComponentType);
+	}
+	
+	if (!TakeWithComponents.IsSet())
+	{
+		FString ComponentRootPath;
+		FString FileName;
+		FString Extension;
+		FPaths::Split(ComponentPath, ComponentRootPath, FileName, Extension);
+
+		TakeWithComponents = { ComponentRootPath, "Slate name not determined", -1, {{UnknownComponentName, ComponentType, ComponentPath}} };
+	}
+
+	GroupFoundComponents(OutTakesWithComponentsCandidates, TakeWithComponents.GetValue());
+};
+
+TOptional<UStereoVideoLiveLinkDevice::FTakeWithComponents> UStereoVideoLiveLinkDevice::ExtractTakeComponentsUsingTokens(FString ComponentPath, FString StoragePath, FString Format, ETakeComponentType ComponentType)
+{
+	FStringView InPathStringView(ComponentPath);
+	FString RelativePath = FString(InPathStringView.SubStr(StoragePath.Len(), InPathStringView.Len() - StoragePath.Len()));
+
+	FVideoDeviceTokenParser TokenParser(Format, RelativePath, UE::StereoVideoLiveLinkDevice::Private::Delimiters);
+
+	FString StorageLeafWithRelativePath = FPaths::GetPathLeaf(StoragePath) + RelativePath;
+	FVideoDeviceTokenParser WithLeafTokenParser(Format, StorageLeafWithRelativePath, UE::StereoVideoLiveLinkDevice::Private::Delimiters);
+
+	FString TakePath;
+	FVideoDeviceTokenParser* SuccessfulParser = nullptr;
+	if (TokenParser.Parse())
+	{
+		SuccessfulParser = &TokenParser;
+		TArray<FString> RelativePathParts;
+		FPaths::NormalizeDirectoryName(RelativePath);
+		RelativePath.ParseIntoArray(RelativePathParts, TEXT("/"));
+		TakePath = StoragePath + RelativePathParts[0];
+	}
+	else if (WithLeafTokenParser.Parse())
+	{
+		SuccessfulParser = &WithLeafTokenParser;
+		TakePath = StoragePath;
+	}
+
+	if (SuccessfulParser != nullptr)
+	{
+		return FTakeWithComponents{ TakePath, SuccessfulParser->GetSlateName(), SuccessfulParser->GetTakeNumber(), { { SuccessfulParser->GetName(), ComponentType, ComponentPath } } };
+	}
+
+	return {};
+};
+
+TOptional<UStereoVideoLiveLinkDevice::FTakeWithComponents> UStereoVideoLiveLinkDevice::ExtractTakeComponentsFromDirectoryStructure(FString ComponentPath, FString StoragePath, ETakeComponentType ComponentType)
+{
+	FStringView InPathStringView(ComponentPath);
+	FString RelativePath = FString(InPathStringView.SubStr(StoragePath.Len(), InPathStringView.Len() - StoragePath.Len()));
+
+	FString TakePath;
+	FString TakeName;
+
+	if (ComponentType == ETakeComponentType::VIDEO || ComponentType == ETakeComponentType::AUDIO)
+	{
+		// Use video file container folder name as TakeName
+		// this is for case where user navigate StoragePath to a container folder
+		TArray<FString> Parts;
+		FPaths::NormalizeDirectoryName(RelativePath);
+		RelativePath.ParseIntoArray(Parts, TEXT("/"));
+		if (Parts.Num() > 1)
+		{
+			TakeName = Parts[0];
+			TakePath = StoragePath + "/" + TakeName;
+		}
+		else
+		{
+			// If the above fails, because user navigates StoragePath to TakeFolder, 
+			// we could use specified storage path leaf folder as TakeName
+			TakeName = FPaths::GetPathLeaf(StoragePath);
+			TakePath = StoragePath;
+		}
+	}
+	else if (ComponentType == ETakeComponentType::IMAGE_SEQUENCE)
+	{
+		// Assumption is that frames are stored in separate folders
+		TArray<FString> Parts;
+		FPaths::NormalizeDirectoryName(RelativePath);
+		RelativePath.ParseIntoArray(Parts, TEXT("/"));
+		if (Parts.Num() > 2)
+		{
+			TakeName = Parts[0];
+			TakePath = StoragePath + "/" + TakeName;
+		}
+		else
+		{
+			TakePath = StoragePath;
+			TakeName = FPaths::GetPathLeaf(StoragePath);
+		}
+
+		{
+			FString Filename, Extension;
+			FPaths::Split(ComponentPath, ComponentPath, Filename, Extension);
+		}
+	}
+
+	if (TakeName.IsEmpty())
+	{
+		return {}; // Failed to match TakeName
+	}
+
+	FString SlateName = TakeName;
+	int32 TakeNumber = -1;
+	FString Name;
+
+	return FTakeWithComponents{ TakePath, SlateName, TakeNumber, { { Name, ComponentType, ComponentPath } } };
+
+};
+
+void UStereoVideoLiveLinkDevice::GroupFoundComponents(TMap<FString, UStereoVideoLiveLinkDevice::FTakeWithComponents>& TakesWithComponentsCandidates, FTakeWithComponents TakeWithComponents)
+{
+	FString SlateName = TakeWithComponents.SlateName;
+
+	FTakeWithComponents* TakeWithComponentsPtr = TakesWithComponentsCandidates.Find(SlateName);
+	if (TakeWithComponentsPtr == nullptr)
+	{
+		TakeWithComponentsPtr = &TakesWithComponentsCandidates.Add(SlateName, { TakeWithComponents.TakePath, TakeWithComponents.SlateName, TakeWithComponents.TakeNumber, {} });
+	}
+
+	for (FTakeWithComponents::Component& Component : TakeWithComponents.Components)
+	{
+		FTakeWithComponents::Component* TakeComponentPtr = TakeWithComponentsPtr->Components.FindByPredicate(
+			[&Component](const FTakeWithComponents::Component& Other)
+			{
+				return Other.Path == Component.Path;
+			});
+
+		if (TakeComponentPtr == nullptr)
+		{
+			TakeWithComponentsPtr->Components.Add(MoveTemp(Component));
+		}
+	}
+};
+
+ELiveLinkDeviceConnectionStatus UStereoVideoLiveLinkDevice::GetConnectionStatus_Implementation() const
+{
+	const UStereoVideoLiveLinkDeviceSettings* DeviceSettings = GetSettings();
+	FString Path = DeviceSettings->StoragePath.Path;
+
+	if (!Path.IsEmpty() && FPaths::DirectoryExists(Path))
+	{
+		return ELiveLinkDeviceConnectionStatus::Connected;
+	}
+
+	return ELiveLinkDeviceConnectionStatus::Disconnected;
+}
+
+FString UStereoVideoLiveLinkDevice::GetHardwareId_Implementation() const
+{
+	return FPlatformMisc::GetDeviceId();
+}
+
+bool UStereoVideoLiveLinkDevice::SetHardwareId_Implementation(const FString& HardwareID)
+{
+	return false;
+}
+
+bool UStereoVideoLiveLinkDevice::Connect_Implementation()
+{
+	const UStereoVideoLiveLinkDeviceSettings* DeviceSettings = GetSettings();
+	FString Path = DeviceSettings->StoragePath.Path;
+
+	if (!Path.IsEmpty() && FPaths::DirectoryExists(Path))
+	{
+		SetConnectionStatus(ELiveLinkDeviceConnectionStatus::Connected);
+		return true;
+	}
+
+	return false;
+}
+
+bool UStereoVideoLiveLinkDevice::Disconnect_Implementation()
+{
+	SetConnectionStatus(ELiveLinkDeviceConnectionStatus::Disconnected);
+	return true;
+}
+
+bool UStereoVideoLiveLinkDevice::IsVideoFile(const FString& InFileNameWithExtension)
+{
+	FString Extension = FPaths::GetExtension(InFileNameWithExtension);
+
+	return UE::StereoVideoLiveLinkDevice::Private::SupportedVideoExtensions.Contains(Extension);
+}
+
+bool UStereoVideoLiveLinkDevice::IsFrameInSequenceFile(const FString& InFileNameWithExtension)
+{
+	FString Extension = FPaths::GetExtension(InFileNameWithExtension);
+
+	return UE::StereoVideoLiveLinkDevice::Private::SupportedImageSequenceExtensions.Contains(Extension);
+}
+
+bool UStereoVideoLiveLinkDevice::IsAudioFile(const FString& InFileNameWithExtension)
+{
+	FString Extension = FPaths::GetExtension(InFileNameWithExtension);
+
+	return UE::StereoVideoLiveLinkDevice::Private::SupportedAudioExtensions.Contains(Extension);
+}
+
+TOptional<FTakeMetadata> UStereoVideoLiveLinkDevice::CreateTakeMetadata(const FTakeWithComponents& InTakeWithComponents) const
+{
+	using namespace UE::CaptureManager;
+
+	FMediaRWManager& MediaManager = UE::StereoVideoLiveLinkDevice::Private::GetMediaRWManager();
+
+	FTakeMetadata TakeMetadata;
+
+	int32 TakeNumber = InTakeWithComponents.TakeNumber;
+
+	if (TakeNumber == INDEX_NONE)
+	{
+		TakeNumber = 1;
+	}
+
+	TakeMetadata.Version.Major = 4;
+	TakeMetadata.Version.Minor = 1;
+
+	TakeMetadata.Slate = InTakeWithComponents.SlateName;
+	TakeMetadata.TakeNumber = TakeNumber;
+
+	// Create new guid
+	TakeMetadata.UniqueId = FGuid::NewGuid().ToString(EGuidFormats::DigitsWithHyphensLower);
+
+	FString InCurrentTakeFile = InTakeWithComponents.Components[0].Path;
+	FFileStatData FileData = IFileManager::Get().GetStatData(*InCurrentTakeFile);
+	TakeMetadata.DateTime = FileData.CreationTime;
+
+	TakeMetadata.Device.Model = TEXT("StereoHMC");
+
+	int32 AudioNameCounter = 1;
+
+	for (const FTakeWithComponents::Component& TakeComponent : InTakeWithComponents.Components)
+	{
+		if (TakeComponent.Type == ETakeComponentType::VIDEO || TakeComponent.Type == ETakeComponentType::IMAGE_SEQUENCE)
+		{
+			FTakeMetadata::FVideo Video;
+
+			FString Name = TakeComponent.Name;
+			if (Name.IsEmpty())
+			{
+				Name = TakeComponent.Path.Mid(InTakeWithComponents.TakePath.Len() + 1);
+				SanitizePackagePath(Name, '_');
+			}
+
+			Video.Name = MoveTemp(Name);
+			Video.Format = FPaths::GetExtension(TakeComponent.Path);
+
+			if (TakeComponent.Type == ETakeComponentType::VIDEO)
+			{
+				Video.PathType = FTakeMetadata::FVideo::EPathType::File;
+
+				TValueOrError<TUniquePtr<IVideoReader>, FText> VideoReaderResult = MediaManager.CreateVideoReader(TakeComponent.Path);
+				if (VideoReaderResult.HasValue())
+				{
+					TUniquePtr<IVideoReader> VideoReader = VideoReaderResult.StealValue();
+					Video.FrameRate = static_cast<float>(VideoReader->GetFrameRate().AsDecimal());
+					TValueOrError<TUniquePtr<UE::CaptureManager::FMediaTextureSample>, FText> TextureSampleResult = VideoReader->Next();
+					if (TextureSampleResult.HasValue())
+					{
+						FVideoDeviceThumbnailExtractor ThumbnailExtractor;
+						UE::CaptureManager::FMediaTextureSample* Sample = TextureSampleResult.GetValue().Get();
+
+						TArray<FColor> ThumbnailRawColorData = ThumbnailExtractor.ExtractThumbnail(Sample);
+
+						if (ThumbnailRawColorData.Num() > 0) {
+							TArray<uint8> ThumbnailData;
+							FImageUtils::ThumbnailCompressImageArray(Sample->Dimensions.X, Sample->Dimensions.Y, ThumbnailRawColorData, ThumbnailData);
+
+							TakeMetadata.Thumbnail = ThumbnailData;
+						}
+					}
+				}
+				else
+				{
+					UE_LOG(LogVideoLiveLinkDevice, Warning, TEXT("Failed to determine the frame rate for the video file %s"), *TakeComponent.Path);
+				}
+
+				Video.Path = TakeComponent.Path;
+
+			}
+			else
+			{
+				Video.PathType = FTakeMetadata::FVideo::EPathType::Folder;
+
+				int32 FramesCount = 0;
+				bool FirstFrame = true;
+
+				bool bIterationResult = IFileManager::Get().IterateDirectoryRecursively(*TakeComponent.Path, [&](const TCHAR* InPath, bool bInIsDirectory)
+					{
+						if (FirstFrame)
+						{
+							FirstFrame = false;
+
+							if (!TakeMetadata.Thumbnail.GetThumbnailData().IsSet())
+							{
+								TArray<uint8> ThumbnailData;
+								if (FFileHelper::LoadFileToArray(ThumbnailData, InPath))
+								{
+									TakeMetadata.Thumbnail = ThumbnailData;
+								}
+							}
+
+							FString Path, Filename, Extension;
+							FPaths::Split(InPath, Path, Filename, Extension);
+							Video.Format = Extension;
+
+							Video.Path = TakeComponent.Path;
+							{
+								TObjectPtr<class UTexture2D> Texture = FImageUtils::ImportFileAsTexture2D(InPath);
+								if (Texture)
+								{
+									Video.FrameWidth = Texture->GetSizeX();
+									Video.FrameHeight = Texture->GetSizeY();
+								}
+							}
+						}
+
+						FramesCount++;
+
+						return true;
+					});
+
+				Video.FramesCount = FramesCount;
+			}
+
+			TakeMetadata.Video.Add(MoveTemp(Video));
+		}
+		else if (TakeComponent.Type == ETakeComponentType::AUDIO)
+		{
+			FTakeMetadata::FAudio Audio;
+			FString Name = TakeComponent.Name;
+			if (Name.IsEmpty())
+			{
+				Name = TEXT("audio") + FString::FromInt(AudioNameCounter++);
+			}
+			Audio.Name = Name;
+			Audio.Path = TakeComponent.Path;
+			Audio.Duration = 0;
+
+			TakeMetadata.Audio.Add(MoveTemp(Audio));
+		}
+	}
+	return TakeMetadata;
+}
+
+#undef LOCTEXT_NAMESPACE
diff --git a/Engine/Plugins/VirtualProduction/CaptureManager/CaptureManagerDevices/Source/StereoVideoLiveLinkDevice/Private/StereoVideoLiveLinkDevice.h b/Engine/Plugins/VirtualProduction/CaptureManager/CaptureManagerDevices/Source/StereoVideoLiveLinkDevice/Private/StereoVideoLiveLinkDevice.h
new file mode 100644
--- /dev/null
+++ b/Engine/Plugins/VirtualProduction/CaptureManager/CaptureManagerDevices/Source/StereoVideoLiveLinkDevice/Private/StereoVideoLiveLinkDevice.h
@@ -0,0 +1,138 @@
+// Copyright Epic Games, Inc. All Rights Reserved.
+
+#pragma once
+
+#include "LiveLinkDevice.h"
+
+#include "Ingest/LiveLinkDeviceCapability_Ingest.h"
+#include "LiveLinkDeviceCapability_Connection.h"
+
+#include "Misc/Optional.h"
+
+#include "Async/TaskProgress.h"
+
+#include "BaseIngestLiveLinkDevice.h"
+
+#include "StereoVideoLiveLinkDevice.generated.h"
+
+
+UCLASS(BlueprintType)
+class UStereoVideoLiveLinkDeviceSettings : public ULiveLinkDeviceSettings
+{
+	GENERATED_BODY()
+
+public:
+	UPROPERTY(EditAnywhere, BlueprintReadWrite, Category = "Stereo Video Device")
+	FString DisplayName = TEXT("Stereo Video Device");
+
+	UPROPERTY(EditAnywhere, BlueprintReadWrite, Category = "Stereo Video Device")
+	FDirectoryPath StoragePath;
+
+	/** 
+	 * Format of the file name that the device expects
+	 *		Available tokens:
+	 *			* {SlateName}
+	 *			* {TakeNumber}
+	 *			* {Name}
+	 *			* {Any}
+	 *		Available delimiters: '_', '-', '.', '\'
+	 * 
+	 */
+	UPROPERTY(EditAnywhere, BlueprintReadWrite, Category = "Stereo Video Device")
+	FString VideoFileNameFormat = TEXT("{Any}");
+
+	/**
+	 * Format of the file name that the device expects
+	 *		Available tokens:
+	 *			* {SlateName}
+	 *			* {TakeNumber}
+	 *			* {Name}
+	 *			* {Any}
+	 *		Available delimiters: '_', '-', '.', '\'
+	 *
+	 */
+	UPROPERTY(EditAnywhere, BlueprintReadWrite, Category = "Stereo Video Device")
+	FString AudioFileNameFormat = TEXT("{Any}");
+};
+
+
+UCLASS(BlueprintType, DisplayName = "Stereo Video Device")
+class STEREOVIDEOLIVELINKDEVICE_API UStereoVideoLiveLinkDevice final : public UBaseIngestLiveLinkDevice
+	, public ILiveLinkDeviceCapability_Connection
+{
+	GENERATED_BODY()
+
+public:
+	UStereoVideoLiveLinkDevice() = default;
+	~UStereoVideoLiveLinkDevice() = default;
+
+	UFUNCTION(BlueprintCallable, Category = "GenericArchiveDeviceSettings")
+	const UStereoVideoLiveLinkDeviceSettings* GetSettings() const;
+
+	//~ Begin ULiveLinkDevice interface
+	virtual TSubclassOf<ULiveLinkDeviceSettings> GetSettingsClass() const override;
+	virtual FText GetDisplayName() const override;
+	virtual EDeviceHealth GetDeviceHealth() const override;
+	virtual FText GetHealthText() const override;
+
+	virtual void OnSettingChanged(const FPropertyChangedEvent& InPropertyChangedEvent) override;
+	//~ End ULiveLinkDevice interface
+
+private:
+
+	enum class ETakeComponentType
+	{
+		VIDEO,
+		IMAGE_SEQUENCE,
+		AUDIO
+	};
+
+	struct FTakeWithComponents
+	{
+		struct Component
+		{
+			FString Name;
+			ETakeComponentType Type;
+			FString Path;
+		};
+
+		FString TakePath;
+
+		FString SlateName;
+		int32 TakeNumber;
+
+		TArray<Component> Components;
+
+		int32 CountComponents(ETakeComponentType Type);
+	};
+	
+	using FTakesWithComponents = TArray<struct FTakeWithComponents>;
+
+	virtual FString GetFullTakePath(UE::CaptureManager::FTakeId InTakeId) const override;
+
+	//~ Begin ULiveLinkDeviceCapability_Ingest interface
+	virtual void RunIngestTake(int32 InTakeId, const UIngestCapability_Options* InIngestOptions) override;
+	virtual void UpdateTakeList_Implementation() override;
+	//~ End ULiveLinkDeviceCapability_Ingest interface
+
+	//~ Begin ILiveLinkDeviceCapability_Connection interface
+	virtual ELiveLinkDeviceConnectionStatus GetConnectionStatus_Implementation() const override;
+	virtual FString GetHardwareId_Implementation() const override;
+	virtual bool SetHardwareId_Implementation(const FString& HardwareID) override;
+	virtual bool Connect_Implementation() override;
+	virtual bool Disconnect_Implementation() override;
+	//~ End ILiveLinkDeviceCapability_Connection interface
+
+	bool IsVideoFile(const FString& InFileNameWithExtension);
+	bool IsFrameInSequenceFile(const FString& InFileNameWithExtension);
+	bool IsAudioFile(const FString& InFileNameWithExtension);
+
+	TOptional<FTakeWithComponents> ExtractTakeComponentsUsingTokens(FString ComponentPath, FString StoragePath, FString Format, ETakeComponentType ComponentType);
+	TOptional<FTakeWithComponents> ExtractTakeComponentsFromDirectoryStructure(FString ComponentPath, FString StoragePath, ETakeComponentType ComponentType);
+	void GroupFoundComponents(TMap<FString, UStereoVideoLiveLinkDevice::FTakeWithComponents>& TakesWithComponentsCandidates, FTakeWithComponents TakeWithComponents);
+	void ExtractTakeComponents(FString ComponentPath, FString StoragePath, ETakeComponentType ComponentType, FString Format, FString UnknownComponentName, TMap<FString, UStereoVideoLiveLinkDevice::FTakeWithComponents>& OutTakesWithComponentsCandidates);
+	TMap<FString, FTakeWithComponents> DiscoverTakes(FString StoragePath);
+	TOptional<FTakeMetadata> CreateTakeMetadata(const FTakeWithComponents & InTakesWithComponents) const;
+
+	TMap<int32, FString> FullTakePaths;
+};
diff --git a/Engine/Plugins/VirtualProduction/CaptureManager/CaptureManagerDevices/Source/StereoVideoLiveLinkDevice/Private/StereoVideoLiveLinkDeviceModule.cpp b/Engine/Plugins/VirtualProduction/CaptureManager/CaptureManagerDevices/Source/StereoVideoLiveLinkDevice/Private/StereoVideoLiveLinkDeviceModule.cpp
new file mode 100644
--- /dev/null
+++ b/Engine/Plugins/VirtualProduction/CaptureManager/CaptureManagerDevices/Source/StereoVideoLiveLinkDevice/Private/StereoVideoLiveLinkDeviceModule.cpp
@@ -0,0 +1,5 @@
+// Copyright Epic Games, Inc. All Rights Reserved.
+
+#include "Modules/ModuleManager.h"
+
+IMPLEMENT_MODULE(FDefaultModuleImpl, StereoVideoLiveLinkDevice)
\ No newline at end of file
diff --git a/Engine/Plugins/VirtualProduction/CaptureManager/CaptureManagerDevices/Source/StereoVideoLiveLinkDevice/StereoVideoLiveLinkDevice.Build.cs b/Engine/Plugins/VirtualProduction/CaptureManager/CaptureManagerDevices/Source/StereoVideoLiveLinkDevice/StereoVideoLiveLinkDevice.Build.cs
new file mode 100644
--- /dev/null
+++ b/Engine/Plugins/VirtualProduction/CaptureManager/CaptureManagerDevices/Source/StereoVideoLiveLinkDevice/StereoVideoLiveLinkDevice.Build.cs
@@ -0,0 +1,42 @@
+// Copyright Epic Games, Inc. All Rights Reserved.
+
+using UnrealBuildTool;
+using System.IO;
+
+public class StereoVideoLiveLinkDevice : ModuleRules
+{
+	public StereoVideoLiveLinkDevice(ReadOnlyTargetRules Target) : base(Target)
+	{
+		PCHUsage = PCHUsageMode.UseExplicitOrSharedPCHs;
+		ShortName = "StereoVideoLLD";
+
+		PublicDependencyModuleNames.AddRange(new string[]
+		{
+			"CaptureUtils",
+			"Core"
+		});
+
+		PrivateDependencyModuleNames.AddRange(new string[]
+		{
+			"Engine",
+			"Json",
+			"CaptureManagerMediaRW",
+			"IngestLiveLinkDevice",
+			"LiveLinkDevice",
+			"LiveLinkCapabilities",
+			"Media",
+			"LiveLinkHubCaptureMessaging",
+			"LiveLinkHub",
+			"ToolWidgets",
+			"Slate",
+			"SlateCore",
+			"CaptureManagerTakeMetadata",
+			"CaptureManagerSettings",
+			"CoreUObject",
+			"UnrealEd",
+			"VideoLiveLinkDeviceCommon"
+		});
+
+		CppCompileWarningSettings.UnsafeTypeCastWarningLevel = WarningLevel.Error;
+	}
+}
diff --git a/Engine/Plugins/VirtualProduction/CaptureManager/CaptureManagerDevices/Source/VideoLiveLinkDeviceCommon/Private/Tests/TestVideoDeviceTokenParser.cpp b/Engine/Plugins/VirtualProduction/CaptureManager/CaptureManagerDevices/Source/VideoLiveLinkDeviceCommon/Private/Tests/TestVideoDeviceTokenParser.cpp
new file mode 100644
--- /dev/null
+++ b/Engine/Plugins/VirtualProduction/CaptureManager/CaptureManagerDevices/Source/VideoLiveLinkDeviceCommon/Private/Tests/TestVideoDeviceTokenParser.cpp
@@ -0,0 +1,446 @@
+// Copyright Epic Games, Inc. All Rights Reserved.
+
+#include "Utils/VideoDeviceTokenParser.h"
+
+#include "Misc/AutomationTest.h"
+
+static const TArray<FString::ElementType> Delimiters =
+{
+	'-',
+	'_',
+	'.',
+	'/'
+};
+
+IMPLEMENT_SIMPLE_AUTOMATION_TEST(FVideoDeviceTokenParserTest1, "CaptureManager.VideoLiveLinkDeviceCommon.TokenParser.Success",
+								 EAutomationTestFlags::EditorContext | EAutomationTestFlags::EngineFilter);
+
+bool FVideoDeviceTokenParserTest1::RunTest(const FString& Parameters)
+{
+	FString InputFormat = TEXT("{SlateName}_{TakeNumber}_{Name}");
+
+	FString SlateName = TEXT("SlateName");
+	int32 TakeNumber = 10;
+	FString Name = TEXT("Name");
+
+	FString FormattedValue = FString::Format(TEXT("{0}_{1}_{2}"), { SlateName, TakeNumber, Name });
+
+	FVideoDeviceTokenParser Parser(MoveTemp(InputFormat), MoveTemp(FormattedValue), Delimiters);
+	TestTrue(TEXT("Parse succeeded"), Parser.Parse());
+	TestEqual(TEXT("Slate Name matches"), SlateName, Parser.GetSlateName());
+	TestEqual(TEXT("Take Number matches"), TakeNumber, Parser.GetTakeNumber());
+	TestEqual(TEXT("Name matches"), Name, Parser.GetName());
+
+	return true;
+}
+
+IMPLEMENT_SIMPLE_AUTOMATION_TEST(FVideoDeviceTokenParserTest2, "CaptureManager.VideoLiveLinkDeviceCommon.TokenParser.Fail",
+								 EAutomationTestFlags::EditorContext | EAutomationTestFlags::EngineFilter);
+
+bool FVideoDeviceTokenParserTest2::RunTest(const FString& Parameters)
+{
+	FString InputFormat = TEXT("{TakeNumber}_{Name}");
+
+	FString SlateName = TEXT("SlateName");
+	int32 TakeNumber = 10;
+	FString Name = TEXT("Name");
+
+	FString FormattedValue = FString::Format(TEXT("{0}_{1}_{2}"), { SlateName, TakeNumber, Name });
+
+	FVideoDeviceTokenParser Parser(MoveTemp(InputFormat), MoveTemp(FormattedValue), Delimiters);
+	TestFalse(TEXT("Parse failed"), Parser.Parse());
+
+	return true;
+}
+
+IMPLEMENT_SIMPLE_AUTOMATION_TEST(FVideoDeviceTokenParserTest3, "CaptureManager.VideoLiveLinkDeviceCommon.TokenParser.Postfix.SuccessAny",
+								 EAutomationTestFlags::EditorContext | EAutomationTestFlags::EngineFilter);
+
+bool FVideoDeviceTokenParserTest3::RunTest(const FString& Parameters)
+{
+	FString InputFormat = TEXT("{SlateName}_{Name}_{Any}");
+
+	FString SlateName = TEXT("SlateName");
+	FString Name = TEXT("Name");
+
+	FString FormattedValue = FString::Format(TEXT("{0}_{1}_TestingAnythingHere"), { SlateName, Name });
+
+	FVideoDeviceTokenParser Parser(MoveTemp(InputFormat), MoveTemp(FormattedValue), Delimiters);
+	TestTrue(TEXT("Parse succeeded"), Parser.Parse());
+	TestEqual(TEXT("Slate Name matches"), SlateName, Parser.GetSlateName());
+	TestEqual(TEXT("Name matches"), Name, Parser.GetName());
+
+	return true;
+}
+
+IMPLEMENT_SIMPLE_AUTOMATION_TEST(FVideoDeviceTokenParserTest4, "CaptureManager.VideoLiveLinkDeviceCommon.TokenParser.Postfix.FailAny",
+								 EAutomationTestFlags::EditorContext | EAutomationTestFlags::EngineFilter);
+
+bool FVideoDeviceTokenParserTest4::RunTest(const FString& Parameters)
+{
+	FString InputFormat = TEXT("{SlateName}_{Name}_");
+
+	FString SlateName = TEXT("SlateName");
+	FString Name = TEXT("Name");
+
+	FString FormattedValue = FString::Format(TEXT("{0}_{1}_TestingAnythingHere"), { SlateName, Name });
+
+	FVideoDeviceTokenParser Parser(MoveTemp(InputFormat), MoveTemp(FormattedValue), Delimiters);
+	TestFalse(TEXT("Parse failed"), Parser.Parse());
+
+	return true;
+}
+
+IMPLEMENT_SIMPLE_AUTOMATION_TEST(FVideoDeviceTokenParserTest5, "CaptureManager.VideoLiveLinkDeviceCommon.TokenParser.Postfix.SuccessCustom",
+								 EAutomationTestFlags::EditorContext | EAutomationTestFlags::EngineFilter);
+
+bool FVideoDeviceTokenParserTest5::RunTest(const FString& Parameters)
+{
+	FString CustomString = TEXT("CustomString");
+
+	FString InputFormat = FString::Format(TEXT("{SlateName}_{Name}_{0}"), { CustomString });
+
+	FString SlateName = TEXT("SlateName");
+	FString Name = TEXT("Name");
+
+	FString FormattedValue = FString::Format(TEXT("{0}_{1}_{2}"), { SlateName, Name, CustomString });
+
+	FVideoDeviceTokenParser Parser(MoveTemp(InputFormat), MoveTemp(FormattedValue), Delimiters);
+	TestTrue(TEXT("Parse succeeded"), Parser.Parse());
+	TestEqual(TEXT("Slate Name matches"), SlateName, Parser.GetSlateName());
+	TestEqual(TEXT("Name matches"), Name, Parser.GetName());
+
+	return true;
+}
+
+IMPLEMENT_SIMPLE_AUTOMATION_TEST(FVideoDeviceTokenParserTest6, "CaptureManager.VideoLiveLinkDeviceCommon.TokenParser.Postfix.FailCustom",
+								 EAutomationTestFlags::EditorContext | EAutomationTestFlags::EngineFilter);
+
+bool FVideoDeviceTokenParserTest6::RunTest(const FString& Parameters)
+{
+	FString CustomString = TEXT("CustomString");
+
+	FString InputFormat = FString::Format(TEXT("{SlateName}_{Name}_{0}"), { CustomString });
+
+	FString SlateName = TEXT("SlateName");
+	FString Name = TEXT("Name");
+
+	FString InvalidCustomString = TEXT("InvalidCustomString");
+	FString FormattedValue = FString::Format(TEXT("{0}_{1}_{2}"), { SlateName, Name, InvalidCustomString });
+
+	FVideoDeviceTokenParser Parser(MoveTemp(InputFormat), MoveTemp(FormattedValue), Delimiters);
+	TestFalse(TEXT("Parse fail"), Parser.Parse());
+
+	return true;
+}
+
+IMPLEMENT_SIMPLE_AUTOMATION_TEST(FVideoDeviceTokenParserTest7, "CaptureManager.VideoLiveLinkDeviceCommon.TokenParser.Prefix.SuccessAny",
+								 EAutomationTestFlags::EditorContext | EAutomationTestFlags::EngineFilter);
+
+bool FVideoDeviceTokenParserTest7::RunTest(const FString& Parameters)
+{
+	FString InputFormat = TEXT("{Any}_{SlateName}_{Name}");
+
+	FString SlateName = TEXT("SlateName");
+	FString Name = TEXT("Name");
+
+	FString FormattedValue = FString::Format(TEXT("TestingAnythingHere_{0}_{1}"), { SlateName, Name });
+
+	FVideoDeviceTokenParser Parser(MoveTemp(InputFormat), MoveTemp(FormattedValue), Delimiters);
+	TestTrue(TEXT("Parse succeeded"), Parser.Parse());
+	TestEqual(TEXT("Slate Name matches"), SlateName, Parser.GetSlateName());
+	TestEqual(TEXT("Name matches"), Name, Parser.GetName());
+
+	return true;
+}
+
+IMPLEMENT_SIMPLE_AUTOMATION_TEST(FVideoDeviceTokenParserTest8, "CaptureManager.VideoLiveLinkDeviceCommon.TokenParser.Prefix.FailAny",
+								 EAutomationTestFlags::EditorContext | EAutomationTestFlags::EngineFilter);
+
+bool FVideoDeviceTokenParserTest8::RunTest(const FString& Parameters)
+{
+	FString InputFormat = TEXT("_{SlateName}_{Name}");
+
+	FString SlateName = TEXT("SlateName");
+	FString Name = TEXT("Name");
+
+	FString FormattedValue = FString::Format(TEXT("TestingAnythingHere_{0}_{1}"), { SlateName, Name });
+
+	FVideoDeviceTokenParser Parser(MoveTemp(InputFormat), MoveTemp(FormattedValue), Delimiters);
+	TestFalse(TEXT("Parse failed"), Parser.Parse());
+
+	return true;
+}
+
+IMPLEMENT_SIMPLE_AUTOMATION_TEST(FVideoDeviceTokenParserTest9, "CaptureManager.VideoLiveLinkDeviceCommon.TokenParser.Prefix.SuccessCustom",
+								 EAutomationTestFlags::EditorContext | EAutomationTestFlags::EngineFilter);
+
+bool FVideoDeviceTokenParserTest9::RunTest(const FString& Parameters)
+{
+	FString CustomString = TEXT("CustomString");
+
+	FString InputFormat = FString::Format(TEXT("{0}_{SlateName}_{Name}"), { CustomString });
+
+	FString SlateName = TEXT("SlateName");
+	FString Name = TEXT("Name");
+
+	FString FormattedValue = FString::Format(TEXT("{0}_{1}_{2}"), { CustomString, SlateName, Name });
+
+	FVideoDeviceTokenParser Parser(MoveTemp(InputFormat), MoveTemp(FormattedValue), Delimiters);
+	TestTrue(TEXT("Parse succeeded"), Parser.Parse());
+	TestEqual(TEXT("Slate Name matches"), SlateName, Parser.GetSlateName());
+	TestEqual(TEXT("Name matches"), Name, Parser.GetName());
+
+	return true;
+}
+
+IMPLEMENT_SIMPLE_AUTOMATION_TEST(FVideoDeviceTokenParserTest10, "CaptureManager.VideoLiveLinkDeviceCommon.TokenParser.Prefix.FailCustom",
+								 EAutomationTestFlags::EditorContext | EAutomationTestFlags::EngineFilter);
+
+bool FVideoDeviceTokenParserTest10::RunTest(const FString& Parameters)
+{
+	FString CustomString = TEXT("CustomString");
+
+	FString InputFormat = FString::Format(TEXT("{0}_{SlateName}_{Name}"), { CustomString });
+
+	FString SlateName = TEXT("SlateName");
+	FString Name = TEXT("Name");
+
+	FString InvalidCustomString = TEXT("InvalidCustomString");
+	FString FormattedValue = FString::Format(TEXT("{0}_{1}_{2}"), { InvalidCustomString, SlateName, Name });
+
+	FVideoDeviceTokenParser Parser(MoveTemp(InputFormat), MoveTemp(FormattedValue), Delimiters);
+	TestFalse(TEXT("Parse fail"), Parser.Parse());
+
+	return true;
+}
+
+IMPLEMENT_SIMPLE_AUTOMATION_TEST(FVideoDeviceTokenParserTest11, "CaptureManager.VideoLiveLinkDeviceCommon.TokenParser.Mid.SuccessAny",
+								 EAutomationTestFlags::EditorContext | EAutomationTestFlags::EngineFilter);
+
+bool FVideoDeviceTokenParserTest11::RunTest(const FString& Parameters)
+{
+	FString InputFormat = TEXT("{SlateName}_{Any}_{Name}");
+
+	FString SlateName = TEXT("SlateName");
+	FString Name = TEXT("Name");
+
+	FString FormattedValue = FString::Format(TEXT("{0}_TestingAnythingHere_{1}"), { SlateName, Name });
+
+	FVideoDeviceTokenParser Parser(MoveTemp(InputFormat), MoveTemp(FormattedValue), Delimiters);
+	TestTrue(TEXT("Parse succeeded"), Parser.Parse());
+	TestEqual(TEXT("Slate Name matches"), SlateName, Parser.GetSlateName());
+	TestEqual(TEXT("Name matches"), Name, Parser.GetName());
+
+	return true;
+}
+
+IMPLEMENT_SIMPLE_AUTOMATION_TEST(FVideoDeviceTokenParserTest12, "CaptureManager.VideoLiveLinkDeviceCommon.TokenParser.Mid.FailAny",
+								 EAutomationTestFlags::EditorContext | EAutomationTestFlags::EngineFilter);
+
+bool FVideoDeviceTokenParserTest12::RunTest(const FString& Parameters)
+{
+	FString InputFormat = TEXT("{SlateName}_{Name}");
+
+	FString SlateName = TEXT("SlateName");
+	FString Name = TEXT("Name");
+
+	FString FormattedValue = FString::Format(TEXT("{0}_TestingAnythingHere_{1}"), { SlateName, Name });
+
+	FVideoDeviceTokenParser Parser(MoveTemp(InputFormat), MoveTemp(FormattedValue), Delimiters);
+	TestFalse(TEXT("Parse failed"), Parser.Parse());
+
+	return true;
+}
+
+IMPLEMENT_SIMPLE_AUTOMATION_TEST(FVideoDeviceTokenParserTest13, "CaptureManager.VideoLiveLinkDeviceCommon.TokenParser.Mid.SuccessCustom",
+								 EAutomationTestFlags::EditorContext | EAutomationTestFlags::EngineFilter);
+
+bool FVideoDeviceTokenParserTest13::RunTest(const FString& Parameters)
+{
+	FString CustomString = TEXT("CustomString");
+
+	FString InputFormat = FString::Format(TEXT("{SlateName}_{0}_{Name}"), { CustomString });
+
+	FString SlateName = TEXT("SlateName");
+	FString Name = TEXT("Name");
+
+	FString FormattedValue = FString::Format(TEXT("{0}_{1}_{2}"), { SlateName, CustomString, Name });
+
+	FVideoDeviceTokenParser Parser(MoveTemp(InputFormat), MoveTemp(FormattedValue), Delimiters);
+	TestTrue(TEXT("Parse succeeded"), Parser.Parse());
+	TestEqual(TEXT("Slate Name matches"), SlateName, Parser.GetSlateName());
+	TestEqual(TEXT("Name matches"), Name, Parser.GetName());
+
+	return true;
+}
+
+IMPLEMENT_SIMPLE_AUTOMATION_TEST(FVideoDeviceTokenParserTest14, "CaptureManager.VideoLiveLinkDeviceCommon.TokenParser.Mid.FailCustom",
+								 EAutomationTestFlags::EditorContext | EAutomationTestFlags::EngineFilter);
+
+bool FVideoDeviceTokenParserTest14::RunTest(const FString& Parameters)
+{
+	FString CustomString = TEXT("CustomString");
+
+	FString InputFormat = FString::Format(TEXT("{SlateName}_{0}_{Name}"), { CustomString });
+
+	FString SlateName = TEXT("SlateName");
+	FString Name = TEXT("Name");
+
+	FString InvalidCustomString = TEXT("InvalidCustomString");
+	FString FormattedValue = FString::Format(TEXT("{0}_{1}_{2}"), { SlateName, InvalidCustomString, Name });
+
+	FVideoDeviceTokenParser Parser(MoveTemp(InputFormat), MoveTemp(FormattedValue), Delimiters);
+	TestFalse(TEXT("Parse fail"), Parser.Parse());
+
+	return true;
+}
+
+IMPLEMENT_SIMPLE_AUTOMATION_TEST(FVideoDeviceTokenParserTest15, "CaptureManager.VideoLiveLinkDeviceCommon.TokenParser.DifferentDelimiter1Success",
+								 EAutomationTestFlags::EditorContext | EAutomationTestFlags::EngineFilter);
+
+bool FVideoDeviceTokenParserTest15::RunTest(const FString& Parameters)
+{
+	FString InputFormat = TEXT("{SlateName}_{TakeNumber}-{Name}");
+
+	FString SlateName = TEXT("SlateName");
+	int32 TakeNumber = 10;
+	FString Name = TEXT("Name");
+
+	FString FormattedValue = FString::Format(TEXT("{0}_{1}-{2}"), { SlateName, TakeNumber, Name });
+
+	FVideoDeviceTokenParser Parser(MoveTemp(InputFormat), MoveTemp(FormattedValue), Delimiters);
+	TestTrue(TEXT("Parse succeeded"), Parser.Parse());
+
+	return true;
+}
+
+IMPLEMENT_SIMPLE_AUTOMATION_TEST(FVideoDeviceTokenParserTest16, "CaptureManager.VideoLiveLinkDeviceCommon.TokenParser.DifferentDelimiterFailure",
+								 EAutomationTestFlags::EditorContext | EAutomationTestFlags::EngineFilter);
+
+bool FVideoDeviceTokenParserTest16::RunTest(const FString& Parameters)
+{
+	FString InputFormat = TEXT("{SlateName}-{Name}");
+
+	FString SlateName = TEXT("SlateName");
+	int32 TakeNumber = 10;
+	FString Name = TEXT("Name");
+
+	FString FormattedValue = FString::Format(TEXT("{0}_{1}-{2}"), { SlateName, TakeNumber, Name });
+
+	FVideoDeviceTokenParser Parser(MoveTemp(InputFormat), MoveTemp(FormattedValue), Delimiters);
+	TestFalse(TEXT("Parse failed"), Parser.Parse());
+	
+	return true;
+}
+
+IMPLEMENT_SIMPLE_AUTOMATION_TEST(FVideoDeviceTokenParserTest17, "CaptureManager.VideoLiveLinkDeviceCommon.TokenParser.SingleTokenValue",
+								 EAutomationTestFlags::EditorContext | EAutomationTestFlags::EngineFilter);
+
+bool FVideoDeviceTokenParserTest17::RunTest(const FString& Parameters)
+{
+	FString InputFormat = TEXT("{SlateName}");
+
+	FString SlateName = TEXT("SlateName");
+	int32 TakeNumber = 10;
+
+	FString FormattedValue = FString::Format(TEXT("{0}"), { SlateName });
+
+	FVideoDeviceTokenParser Parser(MoveTemp(InputFormat), MoveTemp(FormattedValue), Delimiters);
+	TestTrue(TEXT("Parse succeeded"), Parser.Parse());
+	TestEqual(TEXT("Slate Name matches"), SlateName, Parser.GetSlateName());
+
+	return true;
+}
+
+IMPLEMENT_SIMPLE_AUTOMATION_TEST(FVideoDeviceTokenParserTest18, "CaptureManager.VideoLiveLinkDeviceCommon.TokenParser.SingleValue",
+								 EAutomationTestFlags::EditorContext | EAutomationTestFlags::EngineFilter);
+
+bool FVideoDeviceTokenParserTest18::RunTest(const FString& Parameters)
+{
+	FString InputFormat = TEXT("TestingString");
+	FString TestingString = TEXT("TestingString");
+
+	int32 NoTakeNumber = INDEX_NONE;
+
+	FVideoDeviceTokenParser Parser(MoveTemp(InputFormat), MoveTemp(TestingString), Delimiters);
+	TestTrue(TEXT("Parse succeeded"), Parser.Parse());
+	TestTrue(TEXT("Slate Name empty"), Parser.GetSlateName().IsEmpty());
+	TestEqual(TEXT("Take Number invalid"), NoTakeNumber, Parser.GetTakeNumber());
+	TestTrue(TEXT("Name empty"), Parser.GetName().IsEmpty());
+
+	return true;
+}
+
+IMPLEMENT_SIMPLE_AUTOMATION_TEST(FVideoDeviceTokenParserTest19, "CaptureManager.VideoLiveLinkDeviceCommon.TokenParser.EmptyString",
+								 EAutomationTestFlags::EditorContext | EAutomationTestFlags::EngineFilter);
+
+bool FVideoDeviceTokenParserTest19::RunTest(const FString& Parameters)
+{
+	FString InputFormat = TEXT("");
+
+	FString SlateName = TEXT("SlateName");
+	int32 TakeNumber = 10;
+
+	FString FormattedValue = FString::Format(TEXT("{0}_{1}"), { SlateName, TakeNumber });
+
+	FVideoDeviceTokenParser Parser(MoveTemp(InputFormat), MoveTemp(FormattedValue), Delimiters);
+	TestFalse(TEXT("Parse failed"), Parser.Parse());
+
+	return true;
+}
+
+IMPLEMENT_SIMPLE_AUTOMATION_TEST(FVideoDeviceTokenParserTest20, "CaptureManager.VideoLiveLinkDeviceCommon.TokenParser.MultipleDelimiters",
+								 EAutomationTestFlags::EditorContext | EAutomationTestFlags::EngineFilter);
+
+bool FVideoDeviceTokenParserTest20::RunTest(const FString& Parameters)
+{
+	FString InputFormat = TEXT("______");
+
+	FString SlateName = TEXT("SlateName");
+	int32 TakeNumber = 10;
+
+	FString FormattedValue = FString::Format(TEXT("{0}_{1}"), { SlateName, TakeNumber });
+
+	FVideoDeviceTokenParser Parser(MoveTemp(InputFormat), MoveTemp(FormattedValue), Delimiters);
+	TestFalse(TEXT("Parse failed"), Parser.Parse());
+
+	return true;
+}
+
+IMPLEMENT_SIMPLE_AUTOMATION_TEST(FVideoDeviceTokenParserTest21, "CaptureManager.VideoLiveLinkDeviceCommon.TokenParser.SingleDelimiter",
+								 EAutomationTestFlags::EditorContext | EAutomationTestFlags::EngineFilter);
+
+bool FVideoDeviceTokenParserTest21::RunTest(const FString& Parameters)
+{
+	FString InputFormat = TEXT("_");
+
+	FString SlateName = TEXT("SlateName");
+	int32 TakeNumber = 10;
+
+	FString FormattedValue = FString::Format(TEXT("{0}_{1}"), { SlateName, TakeNumber });
+
+	FVideoDeviceTokenParser Parser(MoveTemp(InputFormat), MoveTemp(FormattedValue), Delimiters);
+	TestFalse(TEXT("Parse failed"), Parser.Parse());
+
+	return true;
+}
+
+IMPLEMENT_SIMPLE_AUTOMATION_TEST(FVideoDeviceTokenParserTest22, "CaptureManager.VideoLiveLinkDeviceCommon.TokenParser.DifferentDelimiter2Success",
+	EAutomationTestFlags::EditorContext | EAutomationTestFlags::EngineFilter);
+
+bool FVideoDeviceTokenParserTest22::RunTest(const FString& Parameters)
+{
+	FString InputFormat = TEXT("{SlateName}/{TakeNumber}-{Name}");
+
+	FString SlateName = TEXT("SlateName");
+	int32 TakeNumber = 10;
+	FString Name = TEXT("Name");
+
+	FString FormattedValue = FString::Format(TEXT("{0}/{1}-{2}"), { SlateName, TakeNumber, Name });
+
+	FVideoDeviceTokenParser Parser(MoveTemp(InputFormat), MoveTemp(FormattedValue), Delimiters);
+	TestTrue(TEXT("Parse succeeded"), Parser.Parse());
+
+	return true;
+}
diff --git a/Engine/Plugins/VirtualProduction/CaptureManager/CaptureManagerDevices/Source/VideoLiveLinkDeviceCommon/Private/Utils/VideoDeviceThumbnailExtractor.cpp b/Engine/Plugins/VirtualProduction/CaptureManager/CaptureManagerDevices/Source/VideoLiveLinkDeviceCommon/Private/Utils/VideoDeviceThumbnailExtractor.cpp
new file mode 100644
--- /dev/null
+++ b/Engine/Plugins/VirtualProduction/CaptureManager/CaptureManagerDevices/Source/VideoLiveLinkDeviceCommon/Private/Utils/VideoDeviceThumbnailExtractor.cpp
@@ -0,0 +1,83 @@
+// Copyright Epic Games, Inc. All Rights Reserved.
+
+#include "Utils/VideoDeviceThumbnailExtractor.h"
+#include "ImageUtils.h"
+#include "MediaSample.h"
+#include "Engine/Texture2D.h"
+#include "CaptureManagerMediaRWModule.h"
+#include "VideoLiveLinkDeviceLog.h"
+
+DEFINE_LOG_CATEGORY(LogVideoLiveLinkDevice);
+
+
+namespace UE::CaptureManager
+{
+
+
+FVideoDeviceThumbnailExtractor::FVideoDeviceThumbnailExtractor()
+{
+}
+
+TArray<FColor> FVideoDeviceThumbnailExtractor::ExtractThumbnail(UE::CaptureManager::FMediaTextureSample* Sample)
+{
+	EMediaTexturePixelFormat SampleFormat = Sample->CurrentFormat;
+	TArray<uint8>& Buffer = Sample->Buffer;
+
+	TArray<FColor> ThumbnailRawColorData;
+	switch (SampleFormat)
+	{
+	case EMediaTexturePixelFormat::U8_Mono:
+	{
+		ThumbnailRawColorData.Reserve(Buffer.Num());
+		for (uint8& value : Buffer)
+		{
+			ThumbnailRawColorData.Add({ value, value, value, 0 });
+		}
+		break;
+	}
+	case EMediaTexturePixelFormat::U8_RGB:
+	{
+		ThumbnailRawColorData.Reserve(Buffer.Num() / 3);
+		for (int32 ValueIndex = 0; ValueIndex < Buffer.Num(); ValueIndex += 3)
+		{
+			ThumbnailRawColorData.Add({ Buffer[ValueIndex], Buffer[ValueIndex + 1], Buffer[ValueIndex + 2], 0 });
+		}
+		break;
+	}
+	case EMediaTexturePixelFormat::U8_BGR:
+	{
+		ThumbnailRawColorData.Reserve(Buffer.Num() / 3);
+		for (int32 ValueIndex = 0; ValueIndex < Buffer.Num(); ValueIndex += 3)
+		{
+			ThumbnailRawColorData.Add({ Buffer[ValueIndex + 2], Buffer[ValueIndex + 1], Buffer[ValueIndex], 0 });
+		}
+		break;
+	}
+	case EMediaTexturePixelFormat::U8_RGBA:
+	{
+		ThumbnailRawColorData.Reserve(Buffer.Num() / 4);
+		for (int32 ValueIndex = 0; ValueIndex < Buffer.Num(); ValueIndex += 4)
+		{
+			ThumbnailRawColorData.Add({ Buffer[ValueIndex], Buffer[ValueIndex + 1], Buffer[ValueIndex + 2], Buffer[ValueIndex + 3] });
+		}
+		break;
+	}
+	case EMediaTexturePixelFormat::U8_BGRA:
+	{
+		ThumbnailRawColorData.Reserve(Buffer.Num() / 4);
+		for (int32 ValueIndex = 0; ValueIndex < Buffer.Num(); ValueIndex += 4)
+		{
+			ThumbnailRawColorData.Add({ Buffer[ValueIndex + 2], Buffer[ValueIndex + 1], Buffer[ValueIndex], Buffer[ValueIndex + 3] });
+		}
+		break;
+	}
+	default:
+	{
+		UE_LOG(LogVideoLiveLinkDevice, Warning, TEXT("Unsupported image format"));
+	}
+	}
+
+	return ThumbnailRawColorData;
+}
+
+}
\ No newline at end of file
diff --git a/Engine/Plugins/VirtualProduction/CaptureManager/CaptureManagerDevices/Source/VideoLiveLinkDeviceCommon/Private/Utils/VideoDeviceTokenParser.cpp b/Engine/Plugins/VirtualProduction/CaptureManager/CaptureManagerDevices/Source/VideoLiveLinkDeviceCommon/Private/Utils/VideoDeviceTokenParser.cpp
new file mode 100644
--- /dev/null
+++ b/Engine/Plugins/VirtualProduction/CaptureManager/CaptureManagerDevices/Source/VideoLiveLinkDeviceCommon/Private/Utils/VideoDeviceTokenParser.cpp
@@ -0,0 +1,173 @@
+// Copyright Epic Games, Inc. All Rights Reserved.
+
+#include "Utils/VideoDeviceTokenParser.h"
+
+#include "Containers/Map.h"
+#include "Containers/Array.h"
+#include "Containers/StringView.h"
+
+const FString FVideoDeviceTokenParser::SlateNameToken = TEXT("{SlateName}");
+const FString FVideoDeviceTokenParser::TakeNumberToken = TEXT("{TakeNumber}");
+const FString FVideoDeviceTokenParser::NameToken = TEXT("{Name}");
+const FString FVideoDeviceTokenParser::AnyToken = TEXT("{Any}");
+
+namespace UE::VideoLiveLinkDevice::Private
+{
+
+TOptional<FString::ElementType> FindDelimiter(const FString& InFileName, const TArray<FString::ElementType>& InDelimiters)
+{
+	for (const FString::ElementType Character : InFileName)
+	{
+		if (InDelimiters.Contains(Character))
+		{
+			return Character;
+		}
+	}
+
+	return {};
+}
+
+}
+
+FVideoDeviceTokenParser::FVideoDeviceTokenParser(const FString& InFormat, const FString& InFormattedValue, const TArray<FString::ElementType>& InDelimiters)
+	: Format(InFormat)
+	, FormattedValue(InFormattedValue)
+	, Delimiters(InDelimiters)
+{
+}
+
+bool FVideoDeviceTokenParser::Parse()
+{
+	TOptional<FString::ElementType> Delimiter = UE::VideoLiveLinkDevice::Private::FindDelimiter(Format, Delimiters);
+
+	TOptional<FComponentValueMap> ComponentValueMap;
+	if (Delimiter.IsSet())
+	{
+		ComponentValueMap = ParseWithDelimiters();
+	}
+	else
+	{
+		ComponentValueMap = ParseWithoutDelimiters();
+	}
+
+	if (!ComponentValueMap.IsSet())
+	{
+		return false;
+	}
+
+	return DetermineValues(MoveTemp(ComponentValueMap.GetValue()));
+}
+
+TOptional<FVideoDeviceTokenParser::FComponentValueMap> FVideoDeviceTokenParser::ParseWithDelimiters() const
+{
+	TArray<FString> FileNameFormatComponents = ParseWithDelimiters(Format);
+
+	if (FileNameFormatComponents.IsEmpty())
+	{
+		return {};
+	}
+
+	TArray<FString> FileNameComponents = ParseWithDelimiters(FormattedValue);
+
+	if (FileNameComponents.IsEmpty())
+	{
+		return {};
+	}
+
+	if (FileNameFormatComponents.Num() != FileNameComponents.Num())
+	{
+		return {};
+	}
+
+	FComponentValueMap ComponentValueMap;
+	for (int32 Index = 0; Index < FileNameComponents.Num(); ++Index)
+	{
+		ComponentValueMap.Add(FileNameFormatComponents[Index], FileNameComponents[Index]);
+	}
+
+	return ComponentValueMap;
+}
+
+TArray<FString> FVideoDeviceTokenParser::ParseWithDelimiters(FString InString, int32 InDelimiterIndex) const
+{
+	TArray<FString> Result;
+	if (InDelimiterIndex >= Delimiters.Num())
+	{
+		Result.Add(InString);
+		return Result;
+	}
+
+	FString Delimiter;
+	Delimiter.AppendChar(Delimiters[InDelimiterIndex]);
+
+	TArray<FString> HalfwaySplitComponents;
+
+	InString.ParseIntoArray(HalfwaySplitComponents, *Delimiter);
+
+	for (const FString& Component : HalfwaySplitComponents)
+	{
+		for (const FString& Part : ParseWithDelimiters(Component, ++InDelimiterIndex))
+		{
+			Result.Add(Part);
+		}
+	}
+	
+	return Result;
+}
+
+
+TOptional<FVideoDeviceTokenParser::FComponentValueMap> FVideoDeviceTokenParser::ParseWithoutDelimiters() const
+{
+	FComponentValueMap ComponentValueMap;
+	ComponentValueMap.Add(Format, FormattedValue);
+
+	return ComponentValueMap;
+}
+
+bool FVideoDeviceTokenParser::DetermineValues(FComponentValueMap InComponentValueMap)
+{
+	if (const FString* MaybeSlateName = InComponentValueMap.Find(SlateNameToken))
+	{
+		SlateName = *MaybeSlateName;
+		InComponentValueMap.Remove(SlateNameToken);
+	}
+
+	if (const FString* MaybeTakeNumber = InComponentValueMap.Find(TakeNumberToken))
+	{
+		TakeNumber = FCString::Atoi(**MaybeTakeNumber);
+		InComponentValueMap.Remove(TakeNumberToken);
+	}
+
+	if (const FString* MaybeName = InComponentValueMap.Find(NameToken))
+	{
+		Name = *MaybeName;
+		InComponentValueMap.Remove(NameToken);
+	}
+
+	InComponentValueMap.Remove(AnyToken);
+
+	for (const TPair<FString, FString>& AdditionalComponentValue : InComponentValueMap)
+	{
+		if (AdditionalComponentValue.Key != AdditionalComponentValue.Value)
+		{
+			return false;
+		}
+	}
+
+	return true;
+}
+
+FString FVideoDeviceTokenParser::GetSlateName() const
+{
+	return SlateName;
+}
+
+int32 FVideoDeviceTokenParser::GetTakeNumber() const
+{
+	return TakeNumber;
+}
+
+FString FVideoDeviceTokenParser::GetName() const
+{
+	return Name;
+}
diff --git a/Engine/Plugins/VirtualProduction/CaptureManager/CaptureManagerDevices/Source/VideoLiveLinkDeviceCommon/Private/VideoLiveLinkDeviceModule.cpp b/Engine/Plugins/VirtualProduction/CaptureManager/CaptureManagerDevices/Source/VideoLiveLinkDeviceCommon/Private/VideoLiveLinkDeviceModule.cpp
new file mode 100644
--- /dev/null
+++ b/Engine/Plugins/VirtualProduction/CaptureManager/CaptureManagerDevices/Source/VideoLiveLinkDeviceCommon/Private/VideoLiveLinkDeviceModule.cpp
@@ -0,0 +1,5 @@
+// Copyright Epic Games, Inc. All Rights Reserved.
+
+#include "Modules/ModuleManager.h"
+
+IMPLEMENT_MODULE(FDefaultModuleImpl, VideoLiveLinkDeviceCommon)
\ No newline at end of file
diff --git a/Engine/Plugins/VirtualProduction/CaptureManager/CaptureManagerDevices/Source/VideoLiveLinkDeviceCommon/Public/Utils/VideoDeviceThumbnailExtractor.h b/Engine/Plugins/VirtualProduction/CaptureManager/CaptureManagerDevices/Source/VideoLiveLinkDeviceCommon/Public/Utils/VideoDeviceThumbnailExtractor.h
new file mode 100644
--- /dev/null
+++ b/Engine/Plugins/VirtualProduction/CaptureManager/CaptureManagerDevices/Source/VideoLiveLinkDeviceCommon/Public/Utils/VideoDeviceThumbnailExtractor.h
@@ -0,0 +1,16 @@
+// Copyright Epic Games, Inc. All Rights Reserved.
+
+#pragma once
+#include "Math/Color.h"
+
+namespace UE::CaptureManager 
+{
+struct FMediaTextureSample;
+
+class VIDEOLIVELINKDEVICECOMMON_API FVideoDeviceThumbnailExtractor
+{
+public:
+	FVideoDeviceThumbnailExtractor();
+	TArray<FColor> ExtractThumbnail(FMediaTextureSample* Sample);
+};
+}
\ No newline at end of file
diff --git a/Engine/Plugins/VirtualProduction/CaptureManager/CaptureManagerDevices/Source/VideoLiveLinkDeviceCommon/Public/Utils/VideoDeviceTokenParser.h b/Engine/Plugins/VirtualProduction/CaptureManager/CaptureManagerDevices/Source/VideoLiveLinkDeviceCommon/Public/Utils/VideoDeviceTokenParser.h
new file mode 100644
--- /dev/null
+++ b/Engine/Plugins/VirtualProduction/CaptureManager/CaptureManagerDevices/Source/VideoLiveLinkDeviceCommon/Public/Utils/VideoDeviceTokenParser.h
@@ -0,0 +1,41 @@
+// Copyright Epic Games, Inc. All Rights Reserved.
+
+#pragma once
+
+#include "Containers/UnrealString.h"
+
+class VIDEOLIVELINKDEVICECOMMON_API FVideoDeviceTokenParser
+{
+public:
+
+	static const FString SlateNameToken;
+	static const FString TakeNumberToken;
+	static const FString NameToken;
+	static const FString AnyToken;
+
+	FVideoDeviceTokenParser(const FString& InFormat, const FString& InFormattedValue, const TArray<FString::ElementType>& InDelimiters);
+
+	bool Parse();
+
+	FString GetSlateName() const;
+	int32 GetTakeNumber() const;
+	FString GetName() const;
+
+private:
+
+	using FComponentValueMap = TMap<FString, FString>;
+
+	TOptional<FComponentValueMap> ParseWithDelimiters() const;
+	TOptional<FComponentValueMap> ParseWithoutDelimiters() const;
+	TArray<FString> ParseWithDelimiters(FString InString, int32 InDelimiterIndex = 0) const;
+
+	bool DetermineValues(FComponentValueMap InComponentValueMap);
+
+	FString Format;
+	FString FormattedValue;
+	const TArray<FString::ElementType>& Delimiters;
+
+	FString SlateName;
+	int32 TakeNumber = INDEX_NONE;
+	FString Name;
+};
\ No newline at end of file
diff --git a/Engine/Plugins/VirtualProduction/CaptureManager/CaptureManagerDevices/Source/VideoLiveLinkDeviceCommon/Public/VideoLiveLinkDeviceLog.h b/Engine/Plugins/VirtualProduction/CaptureManager/CaptureManagerDevices/Source/VideoLiveLinkDeviceCommon/Public/VideoLiveLinkDeviceLog.h
new file mode 100644
--- /dev/null
+++ b/Engine/Plugins/VirtualProduction/CaptureManager/CaptureManagerDevices/Source/VideoLiveLinkDeviceCommon/Public/VideoLiveLinkDeviceLog.h
@@ -0,0 +1,7 @@
+// Copyright Epic Games, Inc. All Rights Reserved.
+
+#pragma once
+
+#include "Logging/LogMacros.h"
+
+DECLARE_LOG_CATEGORY_EXTERN(LogVideoLiveLinkDevice, Log, All);
\ No newline at end of file
diff --git a/Engine/Plugins/VirtualProduction/CaptureManager/CaptureManagerDevices/Source/VideoLiveLinkDeviceCommon/VideoLiveLinkDeviceCommon.Build.cs b/Engine/Plugins/VirtualProduction/CaptureManager/CaptureManagerDevices/Source/VideoLiveLinkDeviceCommon/VideoLiveLinkDeviceCommon.Build.cs
new file mode 100644
--- /dev/null
+++ b/Engine/Plugins/VirtualProduction/CaptureManager/CaptureManagerDevices/Source/VideoLiveLinkDeviceCommon/VideoLiveLinkDeviceCommon.Build.cs
@@ -0,0 +1,26 @@
+// Copyright Epic Games, Inc. All Rights Reserved.
+
+using UnrealBuildTool;
+using System.IO;
+
+public class VideoLiveLinkDeviceCommon : ModuleRules
+{
+	public VideoLiveLinkDeviceCommon(ReadOnlyTargetRules Target) : base(Target)
+	{
+		PCHUsage = PCHUsageMode.UseExplicitOrSharedPCHs;
+		ShortName = "VideoLLDCommon";
+
+		PublicDependencyModuleNames.AddRange(new string[]
+		{
+			"Core"
+		});
+
+		PrivateDependencyModuleNames.AddRange(new string[]
+		{
+			"CaptureManagerMediaRW",
+			"Engine"
+		});
+
+		CppCompileWarningSettings.UnsafeTypeCastWarningLevel = WarningLevel.Error;
+	}
+}
