[CaptureManager] Set default mono video pattern and refresh takes list on property change

diff --git a/Engine/Plugins/VirtualProduction/CaptureManager/CaptureManagerDevices/Source/MonoVideoLiveLinkDevice/Private/MonoVideoLiveLinkDevice.cpp b/Engine/Plugins/VirtualProduction/CaptureManager/CaptureManagerDevices/Source/MonoVideoLiveLinkDevice/Private/MonoVideoLiveLinkDevice.cpp
--- a/Engine/Plugins/VirtualProduction/CaptureManager/CaptureManagerDevices/Source/MonoVideoLiveLinkDevice/Private/MonoVideoLiveLinkDevice.cpp
+++ b/Engine/Plugins/VirtualProduction/CaptureManager/CaptureManagerDevices/Source/MonoVideoLiveLinkDevice/Private/MonoVideoLiveLinkDevice.cpp
@@ -67,7 +67,10 @@ void UMonoVideoLiveLinkDevice::OnSettingChanged(const FPropertyChangedEvent& InP
 {
 	Super::OnSettingChanged(InPropertyChangedEvent);
 
-	if (GET_MEMBER_NAME_CHECKED(UMonoVideoLiveLinkDeviceSettings, StoragePath) == InPropertyChangedEvent.GetMemberPropertyName())
+	const FName& PropertyName = InPropertyChangedEvent.GetMemberPropertyName();
+
+	if (GET_MEMBER_NAME_CHECKED(UMonoVideoLiveLinkDeviceSettings, StoragePath) == PropertyName
+		|| GET_MEMBER_NAME_CHECKED(UMonoVideoLiveLinkDeviceSettings, FileNameFormat) == PropertyName)
 	{
 		FString Path = GetSettings()->StoragePath.Path;
 		if (!Path.IsEmpty() && FPaths::DirectoryExists(Path))
diff --git a/Engine/Plugins/VirtualProduction/CaptureManager/CaptureManagerDevices/Source/MonoVideoLiveLinkDevice/Private/MonoVideoLiveLinkDevice.h b/Engine/Plugins/VirtualProduction/CaptureManager/CaptureManagerDevices/Source/MonoVideoLiveLinkDevice/Private/MonoVideoLiveLinkDevice.h
--- a/Engine/Plugins/VirtualProduction/CaptureManager/CaptureManagerDevices/Source/MonoVideoLiveLinkDevice/Private/MonoVideoLiveLinkDevice.h
+++ b/Engine/Plugins/VirtualProduction/CaptureManager/CaptureManagerDevices/Source/MonoVideoLiveLinkDevice/Private/MonoVideoLiveLinkDevice.h
@@ -40,7 +40,7 @@ public:
 	 * Leave empty to parse every video file in the specified directory
 	 */
 	UPROPERTY(EditAnywhere, BlueprintReadWrite, Category = "Mono Video Device")
-	FString FileNameFormat = TEXT("{SlateName}_{TakeNumber}_{Any}");
+	FString FileNameFormat = TEXT("{Any}");
 };
 
 
