[CaptureManager] Limit which Stereo Video device property changes trigger a take list refresh

diff --git a/Engine/Plugins/VirtualProduction/CaptureManager/CaptureManagerDevices/Source/StereoVideoLiveLinkDevice/Private/StereoVideoLiveLinkDevice.cpp b/Engine/Plugins/VirtualProduction/CaptureManager/CaptureManagerDevices/Source/StereoVideoLiveLinkDevice/Private/StereoVideoLiveLinkDevice.cpp
--- a/Engine/Plugins/VirtualProduction/CaptureManager/CaptureManagerDevices/Source/StereoVideoLiveLinkDevice/Private/StereoVideoLiveLinkDevice.cpp
+++ b/Engine/Plugins/VirtualProduction/CaptureManager/CaptureManagerDevices/Source/StereoVideoLiveLinkDevice/Private/StereoVideoLiveLinkDevice.cpp
@@ -81,14 +81,21 @@ void UStereoVideoLiveLinkDevice::OnSettingChanged(const FPropertyChangedEvent& I
 {
 	Super::OnSettingChanged(InPropertyChangedEvent);
 
-	FString Path = GetSettings()->StoragePath.Path;
-	if (!Path.IsEmpty() && FPaths::DirectoryExists(Path))
-	{
-		SetConnectionStatus(ELiveLinkDeviceConnectionStatus::Connected);
-	}
-	else
+	const FName& PropertyName = InPropertyChangedEvent.GetMemberPropertyName();
+
+	if (GET_MEMBER_NAME_CHECKED(UStereoVideoLiveLinkDeviceSettings, StoragePath) == PropertyName
+		|| GET_MEMBER_NAME_CHECKED(UStereoVideoLiveLinkDeviceSettings, VideoFileNameFormat) == PropertyName
+		|| GET_MEMBER_NAME_CHECKED(UStereoVideoLiveLinkDeviceSettings, AudioFileNameFormat) == PropertyName)
 	{
-		SetConnectionStatus(ELiveLinkDeviceConnectionStatus::Disconnected);
+		FString Path = GetSettings()->StoragePath.Path;
+		if (!Path.IsEmpty() && FPaths::DirectoryExists(Path))
+		{
+			SetConnectionStatus(ELiveLinkDeviceConnectionStatus::Connected);
+		}
+		else
+		{
+			SetConnectionStatus(ELiveLinkDeviceConnectionStatus::Disconnected);
+		}
 	}
 
 }
