[CaptureManager] Renamed StoragePath to TakeDirectory

diff --git a/Engine/Plugins/VirtualProduction/CaptureManager/CaptureManagerDevices/Source/FilesystemLiveLinkDevice/Private/FilesystemLiveLinkDevice.cpp b/Engine/Plugins/VirtualProduction/CaptureManager/CaptureManagerDevices/Source/FilesystemLiveLinkDevice/Private/FilesystemLiveLinkDevice.cpp
--- a/Engine/Plugins/VirtualProduction/CaptureManager/CaptureManagerDevices/Source/FilesystemLiveLinkDevice/Private/FilesystemLiveLinkDevice.cpp
+++ b/Engine/Plugins/VirtualProduction/CaptureManager/CaptureManagerDevices/Source/FilesystemLiveLinkDevice/Private/FilesystemLiveLinkDevice.cpp
@@ -43,9 +43,9 @@ void UFilesystemLiveLinkDevice::OnSettingChanged(const FPropertyChangedEvent& In
 {
 	Super::OnSettingChanged(InPropertyChangedEvent);
 
-	if (GET_MEMBER_NAME_CHECKED(UFilesystemLiveLinkDeviceSettings, StoragePath) == InPropertyChangedEvent.GetMemberPropertyName())
+	if (GET_MEMBER_NAME_CHECKED(UFilesystemLiveLinkDeviceSettings, TakeDirectory) == InPropertyChangedEvent.GetMemberPropertyName())
 	{
-		FString Path = GetSettings()->StoragePath.Path;
+		FString Path = GetSettings()->TakeDirectory.Path;
 		if (!Path.IsEmpty() && FPaths::DirectoryExists(Path))
 		{
 			SetConnectionStatus(ELiveLinkDeviceConnectionStatus::Connected);
@@ -59,7 +59,7 @@ void UFilesystemLiveLinkDevice::OnSettingChanged(const FPropertyChangedEvent& In
 
 FString UFilesystemLiveLinkDevice::GetFullTakePath(UE::CaptureManager::FTakeId InTakeId) const
 {
-	FString StoragePath = GetSettings()->StoragePath.Path;
+	FString StoragePath = GetSettings()->TakeDirectory.Path;
 
 	if (const FString* RelativeTakePath = RelativeTakePaths.Find(InTakeId))
 	{
@@ -94,7 +94,7 @@ void UFilesystemLiveLinkDevice::UpdateTakeList_Implementation()
 	{
 		RemoveAllTakes();
 
-		FString StoragePath = GetSettings()->StoragePath.Path;
+		FString StoragePath = GetSettings()->TakeDirectory.Path;
 
 		TArray<FString> TakeMetadataFiles;
 		const bool bIterationResult = IFileManager::Get().IterateDirectoryRecursively(*StoragePath, [&](const TCHAR* InFileNameOrDirectory, bool bInIsDirectory)
@@ -137,7 +137,7 @@ void UFilesystemLiveLinkDevice::UpdateTakeList_Implementation()
 ELiveLinkDeviceConnectionStatus UFilesystemLiveLinkDevice::GetConnectionStatus_Implementation() const
 {
 	const UFilesystemLiveLinkDeviceSettings* DeviceSettings = GetSettings();
-	FString Path = DeviceSettings->StoragePath.Path;
+	FString Path = DeviceSettings->TakeDirectory.Path;
 
 	if (!Path.IsEmpty() && FPaths::DirectoryExists(Path))
 	{
@@ -160,7 +160,7 @@ bool UFilesystemLiveLinkDevice::SetHardwareId_Implementation(const FString& Hard
 bool UFilesystemLiveLinkDevice::Connect_Implementation()
 {
 	const UFilesystemLiveLinkDeviceSettings* DeviceSettings = GetSettings();
-	FString Path = DeviceSettings->StoragePath.Path;
+	FString Path = DeviceSettings->TakeDirectory.Path;
 
 	if (!Path.IsEmpty() && FPaths::DirectoryExists(Path))
 	{
diff --git a/Engine/Plugins/VirtualProduction/CaptureManager/CaptureManagerDevices/Source/FilesystemLiveLinkDevice/Private/FilesystemLiveLinkDevice.h b/Engine/Plugins/VirtualProduction/CaptureManager/CaptureManagerDevices/Source/FilesystemLiveLinkDevice/Private/FilesystemLiveLinkDevice.h
--- a/Engine/Plugins/VirtualProduction/CaptureManager/CaptureManagerDevices/Source/FilesystemLiveLinkDevice/Private/FilesystemLiveLinkDevice.h
+++ b/Engine/Plugins/VirtualProduction/CaptureManager/CaptureManagerDevices/Source/FilesystemLiveLinkDevice/Private/FilesystemLiveLinkDevice.h
@@ -25,8 +25,9 @@ public:
 	UPROPERTY(EditAnywhere, BlueprintReadWrite, Category = "Filesystem Device")
 	FString DisplayName = TEXT("Filesystem Device");
 
+	/* Path to a directory containing the take(s) data */
 	UPROPERTY(EditAnywhere, BlueprintReadWrite, Category = "Filesystem Device")
-	FDirectoryPath StoragePath;
+	FDirectoryPath TakeDirectory;
 };
 
 
diff --git a/Engine/Plugins/VirtualProduction/CaptureManager/CaptureManagerDevices/Source/MonoVideoLiveLinkDevice/Private/MonoVideoLiveLinkDevice.cpp b/Engine/Plugins/VirtualProduction/CaptureManager/CaptureManagerDevices/Source/MonoVideoLiveLinkDevice/Private/MonoVideoLiveLinkDevice.cpp
--- a/Engine/Plugins/VirtualProduction/CaptureManager/CaptureManagerDevices/Source/MonoVideoLiveLinkDevice/Private/MonoVideoLiveLinkDevice.cpp
+++ b/Engine/Plugins/VirtualProduction/CaptureManager/CaptureManagerDevices/Source/MonoVideoLiveLinkDevice/Private/MonoVideoLiveLinkDevice.cpp
@@ -69,10 +69,10 @@ void UMonoVideoLiveLinkDevice::OnSettingChanged(const FPropertyChangedEvent& InP
 
 	const FName& PropertyName = InPropertyChangedEvent.GetMemberPropertyName();
 
-	if (GET_MEMBER_NAME_CHECKED(UMonoVideoLiveLinkDeviceSettings, StoragePath) == PropertyName
+	if (GET_MEMBER_NAME_CHECKED(UMonoVideoLiveLinkDeviceSettings, TakeDirectory) == PropertyName
 		|| GET_MEMBER_NAME_CHECKED(UMonoVideoLiveLinkDeviceSettings, FileNameFormat) == PropertyName)
 	{
-		FString Path = GetSettings()->StoragePath.Path;
+		FString Path = GetSettings()->TakeDirectory.Path;
 		if (!Path.IsEmpty() && FPaths::DirectoryExists(Path))
 		{
 			SetConnectionStatus(ELiveLinkDeviceConnectionStatus::Connected);
@@ -117,7 +117,7 @@ void UMonoVideoLiveLinkDevice::UpdateTakeList_Implementation()
 {
 	RemoveAllTakes();
 
-	FString StoragePath = GetSettings()->StoragePath.Path;
+	FString StoragePath = GetSettings()->TakeDirectory.Path;
 
 	TArray<FString> SupportedVideoFiles;
 	const bool bIterationResult = IFileManager::Get().IterateDirectoryRecursively(*StoragePath, [&](const TCHAR* InFileNameOrDirectory, bool bInIsDirectory)
@@ -147,7 +147,7 @@ void UMonoVideoLiveLinkDevice::UpdateTakeList_Implementation()
 ELiveLinkDeviceConnectionStatus UMonoVideoLiveLinkDevice::GetConnectionStatus_Implementation() const
 {
 	const UMonoVideoLiveLinkDeviceSettings* DeviceSettings = GetSettings();
-	FString Path = DeviceSettings->StoragePath.Path;
+	FString Path = DeviceSettings->TakeDirectory.Path;
 
 	if (!Path.IsEmpty() && FPaths::DirectoryExists(Path))
 	{
@@ -170,7 +170,7 @@ bool UMonoVideoLiveLinkDevice::SetHardwareId_Implementation(const FString& Hardw
 bool UMonoVideoLiveLinkDevice::Connect_Implementation()
 {
 	const UMonoVideoLiveLinkDeviceSettings* DeviceSettings = GetSettings();
-	FString Path = DeviceSettings->StoragePath.Path;
+	FString Path = DeviceSettings->TakeDirectory.Path;
 
 	if (!Path.IsEmpty() && FPaths::DirectoryExists(Path))
 	{
diff --git a/Engine/Plugins/VirtualProduction/CaptureManager/CaptureManagerDevices/Source/MonoVideoLiveLinkDevice/Private/MonoVideoLiveLinkDevice.h b/Engine/Plugins/VirtualProduction/CaptureManager/CaptureManagerDevices/Source/MonoVideoLiveLinkDevice/Private/MonoVideoLiveLinkDevice.h
--- a/Engine/Plugins/VirtualProduction/CaptureManager/CaptureManagerDevices/Source/MonoVideoLiveLinkDevice/Private/MonoVideoLiveLinkDevice.h
+++ b/Engine/Plugins/VirtualProduction/CaptureManager/CaptureManagerDevices/Source/MonoVideoLiveLinkDevice/Private/MonoVideoLiveLinkDevice.h
@@ -25,8 +25,9 @@ public:
 	UPROPERTY(EditAnywhere, BlueprintReadWrite, Category = "Mono Video Device")
 	FString DisplayName = TEXT("Mono Video Device");
 
+	/* Path to a directory containing the take(s) data */
 	UPROPERTY(EditAnywhere, BlueprintReadWrite, Category = "Mono Video Device")
-	FDirectoryPath StoragePath;
+	FDirectoryPath TakeDirectory;
 
 	/** 
 	 * Format of the file name that the device expects
diff --git a/Engine/Plugins/VirtualProduction/CaptureManager/CaptureManagerDevices/Source/StereoVideoLiveLinkDevice/Private/StereoVideoLiveLinkDevice.cpp b/Engine/Plugins/VirtualProduction/CaptureManager/CaptureManagerDevices/Source/StereoVideoLiveLinkDevice/Private/StereoVideoLiveLinkDevice.cpp
--- a/Engine/Plugins/VirtualProduction/CaptureManager/CaptureManagerDevices/Source/StereoVideoLiveLinkDevice/Private/StereoVideoLiveLinkDevice.cpp
+++ b/Engine/Plugins/VirtualProduction/CaptureManager/CaptureManagerDevices/Source/StereoVideoLiveLinkDevice/Private/StereoVideoLiveLinkDevice.cpp
@@ -83,11 +83,11 @@ void UStereoVideoLiveLinkDevice::OnSettingChanged(const FPropertyChangedEvent& I
 
 	const FName& PropertyName = InPropertyChangedEvent.GetMemberPropertyName();
 
-	if (GET_MEMBER_NAME_CHECKED(UStereoVideoLiveLinkDeviceSettings, StoragePath) == PropertyName
+	if (GET_MEMBER_NAME_CHECKED(UStereoVideoLiveLinkDeviceSettings, TakeDirectory) == PropertyName
 		|| GET_MEMBER_NAME_CHECKED(UStereoVideoLiveLinkDeviceSettings, VideoFileNameFormat) == PropertyName
 		|| GET_MEMBER_NAME_CHECKED(UStereoVideoLiveLinkDeviceSettings, AudioFileNameFormat) == PropertyName)
 	{
-		FString Path = GetSettings()->StoragePath.Path;
+		FString Path = GetSettings()->TakeDirectory.Path;
 		if (!Path.IsEmpty() && FPaths::DirectoryExists(Path))
 		{
 			SetConnectionStatus(ELiveLinkDeviceConnectionStatus::Connected);
@@ -146,7 +146,7 @@ void UStereoVideoLiveLinkDevice::UpdateTakeList_Implementation()
 {
 	RemoveAllTakes();
 
-	FString StoragePath = GetSettings()->StoragePath.Path;
+	FString StoragePath = GetSettings()->TakeDirectory.Path;
 
 	TMap<FString, UStereoVideoLiveLinkDevice::FTakeWithComponents> TakesWithComponentsCandidates = DiscoverTakes(StoragePath);
 
@@ -360,7 +360,7 @@ void UStereoVideoLiveLinkDevice::GroupFoundComponents(TMap<FString, UStereoVideo
 ELiveLinkDeviceConnectionStatus UStereoVideoLiveLinkDevice::GetConnectionStatus_Implementation() const
 {
 	const UStereoVideoLiveLinkDeviceSettings* DeviceSettings = GetSettings();
-	FString Path = DeviceSettings->StoragePath.Path;
+	FString Path = DeviceSettings->TakeDirectory.Path;
 
 	if (!Path.IsEmpty() && FPaths::DirectoryExists(Path))
 	{
@@ -383,7 +383,7 @@ bool UStereoVideoLiveLinkDevice::SetHardwareId_Implementation(const FString& Har
 bool UStereoVideoLiveLinkDevice::Connect_Implementation()
 {
 	const UStereoVideoLiveLinkDeviceSettings* DeviceSettings = GetSettings();
-	FString Path = DeviceSettings->StoragePath.Path;
+	FString Path = DeviceSettings->TakeDirectory.Path;
 
 	if (!Path.IsEmpty() && FPaths::DirectoryExists(Path))
 	{
diff --git a/Engine/Plugins/VirtualProduction/CaptureManager/CaptureManagerDevices/Source/StereoVideoLiveLinkDevice/Private/StereoVideoLiveLinkDevice.h b/Engine/Plugins/VirtualProduction/CaptureManager/CaptureManagerDevices/Source/StereoVideoLiveLinkDevice/Private/StereoVideoLiveLinkDevice.h
--- a/Engine/Plugins/VirtualProduction/CaptureManager/CaptureManagerDevices/Source/StereoVideoLiveLinkDevice/Private/StereoVideoLiveLinkDevice.h
+++ b/Engine/Plugins/VirtualProduction/CaptureManager/CaptureManagerDevices/Source/StereoVideoLiveLinkDevice/Private/StereoVideoLiveLinkDevice.h
@@ -25,8 +25,9 @@ public:
 	UPROPERTY(EditAnywhere, BlueprintReadWrite, Category = "Stereo Video Device")
 	FString DisplayName = TEXT("Stereo Video Device");
 
+	/* Path to a directory containing the take(s) data */
 	UPROPERTY(EditAnywhere, BlueprintReadWrite, Category = "Stereo Video Device")
-	FDirectoryPath StoragePath;
+	FDirectoryPath TakeDirectory;
 
 	/** 
 	 * Format of the file name that the device expects
