Refactor mono and stereo video devices

diff --git a/Engine/Plugins/VirtualProduction/CaptureManager/CaptureManagerDevices/Source/MonoVideoLiveLinkDevice/Private/MonoVideoLiveLinkDevice.cpp b/Engine/Plugins/VirtualProduction/CaptureManager/CaptureManagerDevices/Source/MonoVideoLiveLinkDevice/Private/MonoVideoLiveLinkDevice.cpp
--- a/Engine/Plugins/VirtualProduction/CaptureManager/CaptureManagerDevices/Source/MonoVideoLiveLinkDevice/Private/MonoVideoLiveLinkDevice.cpp
+++ b/Engine/Plugins/VirtualProduction/CaptureManager/CaptureManagerDevices/Source/MonoVideoLiveLinkDevice/Private/MonoVideoLiveLinkDevice.cpp
@@ -8,11 +8,14 @@
 
 #include "CaptureManagerMediaRWModule.h"
 
-#include "Utils/VideoDeviceTokenParser.h"
+#include "Misc/MessageDialog.h"
+#include "Utils/TakeDiscoveryExpressionParser.h"
 #include "Utils/VideoDeviceThumbnailExtractor.h"
 
 #include "VideoLiveLinkDeviceLog.h"
 
+#define LOCTEXT_NAMESPACE "MonoVideoDevice"
+
 DEFINE_LOG_CATEGORY(LogVideoLiveLinkDevice);
 
 namespace UE::MonoVideoLiveLinkDevice::Private
@@ -70,7 +73,7 @@ void UMonoVideoLiveLinkDevice::OnSettingChanged(const FPropertyChangedEvent& InP
 	const FName& PropertyName = InPropertyChangedEvent.GetMemberPropertyName();
 
 	if (GET_MEMBER_NAME_CHECKED(UMonoVideoLiveLinkDeviceSettings, TakeDirectory) == PropertyName
-		|| GET_MEMBER_NAME_CHECKED(UMonoVideoLiveLinkDeviceSettings, FileNameFormat) == PropertyName)
+		|| GET_MEMBER_NAME_CHECKED(UMonoVideoLiveLinkDeviceSettings, VideoDiscoveryExpression) == PropertyName)
 	{
 		FString Path = GetSettings()->TakeDirectory.Path;
 		if (!Path.IsEmpty() && FPaths::DirectoryExists(Path))
@@ -118,10 +121,25 @@ void UMonoVideoLiveLinkDevice::UpdateTakeList_Implementation()
 	RemoveAllTakes();
 
 	FString StoragePath = GetSettings()->TakeDirectory.Path;
-
+	int32 PathsInterrogatedCount = 0;
+	const int32 PathsToInterrogateInOneRun = 100000;
+	
 	TArray<FString> SupportedVideoFiles;
 	const bool bIterationResult = IFileManager::Get().IterateDirectoryRecursively(*StoragePath, [&](const TCHAR* InFileNameOrDirectory, bool bInIsDirectory)
 	{
+		if (++PathsInterrogatedCount > PathsToInterrogateInOneRun)
+		{
+			EAppReturnType::Type Result = FMessageDialog::Open(EAppMsgType::OkCancel, FText::Format(LOCTEXT("MonoVideoDevice_ProceedWithTakesDiscovery", "Analyzed {0} paths in search for takes. Do you want to analyze {1} more?"), PathsToInterrogateInOneRun, PathsToInterrogateInOneRun));
+			if (Result == EAppReturnType::Ok)
+			{
+				PathsInterrogatedCount = 0;
+			}
+			else 
+			{
+				return false;
+			}
+		}
+
 		if (!bInIsDirectory && IsVideoFile(InFileNameOrDirectory))
 		{			
 			SupportedVideoFiles.Add(InFileNameOrDirectory);
@@ -204,16 +222,16 @@ TOptional<FTakeMetadata> UMonoVideoLiveLinkDevice::ReadTake(const FString& InCur
 
 	FFileStatData FileData = IFileManager::Get().GetStatData(*InCurrentTakeFile);
 
-	FString FileNameFormat = GetSettings()->FileNameFormat;
+	FString FileNameFormat = GetSettings()->VideoDiscoveryExpression.Value;
 	FString FileName = FPaths::GetBaseFilename(InCurrentTakeFile);
 
 	FString SlateName;
 	FString Name;
 	int32 TakeNumber = INDEX_NONE;
 
-	if (!FileNameFormat.IsEmpty())
+	if (FileNameFormat != "<Auto>")
 	{
-		FVideoDeviceTokenParser TokenParser(FileNameFormat, FileName, UE::MonoVideoLiveLinkDevice::Private::Delimiters);
+		FTakeDiscoveryExpressionParser TokenParser(FileNameFormat, FileName, UE::MonoVideoLiveLinkDevice::Private::Delimiters);
 		bool bParse = TokenParser.Parse();
 
 		if (!bParse)
@@ -304,3 +322,5 @@ TOptional<FTakeMetadata> UMonoVideoLiveLinkDevice::ReadTake(const FString& InCur
 
 	return TakeMetadata;
 }
+
+#undef LOCTEXT_NAMESPACE
diff --git a/Engine/Plugins/VirtualProduction/CaptureManager/CaptureManagerDevices/Source/MonoVideoLiveLinkDevice/Private/MonoVideoLiveLinkDevice.h b/Engine/Plugins/VirtualProduction/CaptureManager/CaptureManagerDevices/Source/MonoVideoLiveLinkDevice/Private/MonoVideoLiveLinkDevice.h
--- a/Engine/Plugins/VirtualProduction/CaptureManager/CaptureManagerDevices/Source/MonoVideoLiveLinkDevice/Private/MonoVideoLiveLinkDevice.h
+++ b/Engine/Plugins/VirtualProduction/CaptureManager/CaptureManagerDevices/Source/MonoVideoLiveLinkDevice/Private/MonoVideoLiveLinkDevice.h
@@ -6,6 +6,7 @@
 
 #include "Ingest/LiveLinkDeviceCapability_Ingest.h"
 #include "LiveLinkDeviceCapability_Connection.h"
+#include "Customizations/TakeDiscoveryExpressionCustomization.h"
 
 #include "Misc/Optional.h"
 
@@ -22,6 +23,11 @@ class UMonoVideoLiveLinkDeviceSettings : public ULiveLinkDeviceSettings
 	GENERATED_BODY()
 
 public:
+	UMonoVideoLiveLinkDeviceSettings()
+	{
+		VideoDiscoveryExpression = TEXT("<Auto>");
+	}
+
 	UPROPERTY(EditAnywhere, BlueprintReadWrite, Category = "Mono Video Device")
 	FString DisplayName = TEXT("Mono Video Device");
 
@@ -29,19 +35,29 @@ public:
 	UPROPERTY(EditAnywhere, BlueprintReadWrite, Category = "Mono Video Device")
 	FDirectoryPath TakeDirectory;
 
-	/** 
+	/**
 	 * Format of the file name that the device expects
-	 *		Available tokens:
-	 *			* {SlateName}
-	 *			* {TakeNumber}
-	 *			* {Name}
-	 *			* {Any}
-	 *		Available delimiters: '_', '-', '.' 
-	 * 
-	 * Leave empty to parse every video file in the specified directory
+	 *		Available expressions:
+	 *			* <Slate> - slate name
+	 *			* <Name> - identifier for the camera
+	 *			* <Take> - take number
+	 *			* <Any> - any string
+	 *			* <Auto> - used alone to automatically figure out takes based on directory structure
+	 *		Available delimiters: '_', '-', '.'
+	 *
+	 * Examples:
+	 * Considering input video MyTakeFolder/MySlate_MyName_SomeString-005.mov
+	 *
+	 * Video File Name Pattern <Auto> - will set Slate to the video filename “MySlate_MyName_SomeString_005” (dash converted to underline), Name will be set to just "video" and Take will be set to 1 by default.
+	 *
+	 * Video File Name Pattern <Slate>_<Name>_<Any>-<Take> - will extract Slate as “MySlate”, Name as “MyName” and Take as 5. <Any> ensures that the “SomeString” part of the filename is ignored.
+	 *
+	 * In both cases, extension '.mov' is stripped out.
+	 *
+	 * Note: when not using <Auto> expression, then both <Slate> and <Name> must be present in the pattern.
 	 */
 	UPROPERTY(EditAnywhere, BlueprintReadWrite, Category = "Mono Video Device")
-	FString FileNameFormat = TEXT("{Any}");
+	FTakeDiscoveryExpression VideoDiscoveryExpression;
 };
 
 
diff --git a/Engine/Plugins/VirtualProduction/CaptureManager/CaptureManagerDevices/Source/MonoVideoLiveLinkDevice/Private/MonoVideoLiveLinkDeviceModule.cpp b/Engine/Plugins/VirtualProduction/CaptureManager/CaptureManagerDevices/Source/MonoVideoLiveLinkDevice/Private/MonoVideoLiveLinkDeviceModule.cpp
--- a/Engine/Plugins/VirtualProduction/CaptureManager/CaptureManagerDevices/Source/MonoVideoLiveLinkDevice/Private/MonoVideoLiveLinkDeviceModule.cpp
+++ b/Engine/Plugins/VirtualProduction/CaptureManager/CaptureManagerDevices/Source/MonoVideoLiveLinkDevice/Private/MonoVideoLiveLinkDeviceModule.cpp
@@ -2,4 +2,33 @@
 
 #include "Modules/ModuleManager.h"
 
-IMPLEMENT_MODULE(FDefaultModuleImpl, MonoVideoLiveLinkDevice)
\ No newline at end of file
+#include "PropertyEditorModule.h"
+#include "Customizations/TakeDiscoveryExpressionCustomization.h"
+#include "Templates/SharedPointer.h"
+
+class FMonoVideoLiveLinkDeviceModule : public IModuleInterface
+{
+public:
+
+	void StartupModule()
+	{
+		FPropertyEditorModule& PropertyModule = FModuleManager::LoadModuleChecked<FPropertyEditorModule>("PropertyEditor");
+		PropertyModule.RegisterCustomPropertyTypeLayout(FTakeDiscoveryExpression::StaticStruct()->GetFName(),
+			FOnGetPropertyTypeCustomizationInstance::CreateRaw(this, &FMonoVideoLiveLinkDeviceModule::GetTakeDiscoveryExpressionCustomization));
+	}
+
+	void ShutdownModule()
+	{
+		FPropertyEditorModule& PropertyModule = FModuleManager::LoadModuleChecked<FPropertyEditorModule>("PropertyEditor");
+		PropertyModule.UnregisterCustomPropertyTypeLayout(FTakeDiscoveryExpression::StaticStruct()->GetFName());
+	}
+
+private:
+
+	TSharedRef<IPropertyTypeCustomization> GetTakeDiscoveryExpressionCustomization()
+	{
+		return MakeShared<FTakeDiscoveryExpressionCustomization>();
+	}
+};
+
+IMPLEMENT_MODULE(FMonoVideoLiveLinkDeviceModule, MonoVideoLiveLinkDevice)
\ No newline at end of file
diff --git a/Engine/Plugins/VirtualProduction/CaptureManager/CaptureManagerDevices/Source/StereoVideoLiveLinkDevice/Private/StereoVideoLiveLinkDevice.cpp b/Engine/Plugins/VirtualProduction/CaptureManager/CaptureManagerDevices/Source/StereoVideoLiveLinkDevice/Private/StereoVideoLiveLinkDevice.cpp
--- a/Engine/Plugins/VirtualProduction/CaptureManager/CaptureManagerDevices/Source/StereoVideoLiveLinkDevice/Private/StereoVideoLiveLinkDevice.cpp
+++ b/Engine/Plugins/VirtualProduction/CaptureManager/CaptureManagerDevices/Source/StereoVideoLiveLinkDevice/Private/StereoVideoLiveLinkDevice.cpp
@@ -10,7 +10,8 @@
 #include "Misc/FileHelper.h"
 #include "Engine/Texture2D.h"
 
-#include "Utils/VideoDeviceTokenParser.h"
+#include "Misc/MessageDialog.h"
+#include "Utils/TakeDiscoveryExpressionParser.h"
 #include "Utils/VideoDeviceThumbnailExtractor.h"
 
 #include "VideoLiveLinkDeviceLog.h"
@@ -84,8 +85,8 @@ void UStereoVideoLiveLinkDevice::OnSettingChanged(const FPropertyChangedEvent& I
 	const FName& PropertyName = InPropertyChangedEvent.GetMemberPropertyName();
 
 	if (GET_MEMBER_NAME_CHECKED(UStereoVideoLiveLinkDeviceSettings, TakeDirectory) == PropertyName
-		|| GET_MEMBER_NAME_CHECKED(UStereoVideoLiveLinkDeviceSettings, VideoFileNameFormat) == PropertyName
-		|| GET_MEMBER_NAME_CHECKED(UStereoVideoLiveLinkDeviceSettings, AudioFileNameFormat) == PropertyName)
+		|| GET_MEMBER_NAME_CHECKED(UStereoVideoLiveLinkDeviceSettings, VideoDiscoveryExpression) == PropertyName
+		|| GET_MEMBER_NAME_CHECKED(UStereoVideoLiveLinkDeviceSettings, AudioDiscoveryExpression) == PropertyName)
 	{
 		FString Path = GetSettings()->TakeDirectory.Path;
 		if (!Path.IsEmpty() && FPaths::DirectoryExists(Path))
@@ -97,7 +98,6 @@ void UStereoVideoLiveLinkDevice::OnSettingChanged(const FPropertyChangedEvent& I
 			SetConnectionStatus(ELiveLinkDeviceConnectionStatus::Disconnected);
 		}
 	}
-
 }
 
 FString UStereoVideoLiveLinkDevice::GetFullTakePath(UE::CaptureManager::FTakeId InTakeId) const
@@ -185,21 +185,37 @@ TMap<FString, UStereoVideoLiveLinkDevice::FTakeWithComponents> UStereoVideoLiveL
 {
 	TMap<FString, FTakeWithComponents> TakesWithComponentsCandidates;
 
+	int32 PathsInterrogatedCount = 0;
+	const int32 PathsToInterrogateInOneRun = 100000;
+
 	const bool bIterationResult = IFileManager::Get().IterateDirectoryRecursively(*StoragePath, [&](const TCHAR* InPath, bool bInIsDirectory)
 		{
-			FString Extension = FPaths::GetExtension(InPath);
+			if (++PathsInterrogatedCount > PathsToInterrogateInOneRun)
+			{
+				EAppReturnType::Type Result = FMessageDialog::Open(EAppMsgType::OkCancel, FText::Format(LOCTEXT("StereoVideoDevice_ProceedWithTakesDiscovery", "Analyzed {0} paths in search for takes. Do you want to analyze {1} more?"), PathsToInterrogateInOneRun, PathsToInterrogateInOneRun));
+				if (Result == EAppReturnType::Ok)
+				{
+					PathsInterrogatedCount = 0;
+				}
+				else
+				{
+					return false;
+				}
+			}
 
+
+			FString Extension = FPaths::GetExtension(InPath);
 			if (UE::StereoVideoLiveLinkDevice::Private::SupportedVideoExtensions.Contains(Extension))
 			{
-				ExtractTakeComponents(InPath, StoragePath, ETakeComponentType::VIDEO, GetSettings()->VideoFileNameFormat, "UnknownVideoName", TakesWithComponentsCandidates);
+				ExtractTakeComponents(InPath, StoragePath, ETakeComponentType::VIDEO, GetSettings()->VideoDiscoveryExpression.Value, "UnknownVideoName", TakesWithComponentsCandidates);
 			}
 			else if (UE::StereoVideoLiveLinkDevice::Private::SupportedImageSequenceExtensions.Contains(Extension))
 			{
-				ExtractTakeComponents(InPath, StoragePath, ETakeComponentType::IMAGE_SEQUENCE, GetSettings()->VideoFileNameFormat, "UnknownImageSequenceName", TakesWithComponentsCandidates);
+				ExtractTakeComponents(InPath, StoragePath, ETakeComponentType::IMAGE_SEQUENCE, GetSettings()->VideoDiscoveryExpression.Value, "UnknownImageSequenceName", TakesWithComponentsCandidates);
 			}
 			else if (UE::StereoVideoLiveLinkDevice::Private::SupportedAudioExtensions.Contains(Extension))
 			{
-				ExtractTakeComponents(InPath, StoragePath, ETakeComponentType::AUDIO, GetSettings()->AudioFileNameFormat, "UnknownAudioName", TakesWithComponentsCandidates);
+				ExtractTakeComponents(InPath, StoragePath, ETakeComponentType::AUDIO, GetSettings()->AudioDiscoveryExpression.Value, "UnknownAudioName", TakesWithComponentsCandidates);
 			}
 
 			return true;
@@ -212,7 +228,7 @@ void UStereoVideoLiveLinkDevice::ExtractTakeComponents(FString ComponentPath, FS
 {
 	TOptional<FTakeWithComponents> TakeWithComponents;
 
-	if (Format == TEXT("{Any}"))
+	if (Format == TEXT("<Auto>"))
 	{
 		TakeWithComponents = ExtractTakeComponentsFromDirectoryStructure(ComponentPath, StoragePath, ComponentType);
 	}
@@ -238,21 +254,23 @@ TOptional<UStereoVideoLiveLinkDevice::FTakeWithComponents> UStereoVideoLiveLinkD
 {
 	FStringView InPathStringView(ComponentPath);
 	FString RelativePath = FString(InPathStringView.SubStr(StoragePath.Len(), InPathStringView.Len() - StoragePath.Len()));
+	FString Extension = FPaths::GetExtension(RelativePath);
+	FString RelativePathNoExtension = RelativePath.Mid(0, RelativePath.Len() - Extension.Len() - 1);
+	
+	FTakeDiscoveryExpressionParser TokenParser(Format, RelativePathNoExtension, UE::StereoVideoLiveLinkDevice::Private::Delimiters);
 
-	FVideoDeviceTokenParser TokenParser(Format, RelativePath, UE::StereoVideoLiveLinkDevice::Private::Delimiters);
-
-	FString StorageLeafWithRelativePath = FPaths::GetPathLeaf(StoragePath) + RelativePath;
-	FVideoDeviceTokenParser WithLeafTokenParser(Format, StorageLeafWithRelativePath, UE::StereoVideoLiveLinkDevice::Private::Delimiters);
+	FString StorageLeafWithRelativePath = FPaths::GetPathLeaf(StoragePath) + RelativePathNoExtension;
+	FTakeDiscoveryExpressionParser WithLeafTokenParser(Format, StorageLeafWithRelativePath, UE::StereoVideoLiveLinkDevice::Private::Delimiters);
 
 	FString TakePath;
-	FVideoDeviceTokenParser* SuccessfulParser = nullptr;
+	FTakeDiscoveryExpressionParser* SuccessfulParser = nullptr;
 	if (TokenParser.Parse())
 	{
 		SuccessfulParser = &TokenParser;
 		TArray<FString> RelativePathParts;
 		FPaths::NormalizeDirectoryName(RelativePath);
 		RelativePath.ParseIntoArray(RelativePathParts, TEXT("/"));
-		TakePath = StoragePath + RelativePathParts[0];
+		TakePath = StoragePath + TEXT("/") + RelativePathParts[0];
 	}
 	else if (WithLeafTokenParser.Parse())
 	{
@@ -329,17 +347,16 @@ TOptional<UStereoVideoLiveLinkDevice::FTakeWithComponents> UStereoVideoLiveLinkD
 	FString Name;
 
 	return FTakeWithComponents{ TakePath, SlateName, TakeNumber, { { Name, ComponentType, ComponentPath } } };
-
 };
 
 void UStereoVideoLiveLinkDevice::GroupFoundComponents(TMap<FString, UStereoVideoLiveLinkDevice::FTakeWithComponents>& TakesWithComponentsCandidates, FTakeWithComponents TakeWithComponents)
 {
-	FString SlateName = TakeWithComponents.SlateName;
+	FString TakeIdentifier = TakeWithComponents.SlateName + FString::FromInt(TakeWithComponents.TakeNumber);
 
-	FTakeWithComponents* TakeWithComponentsPtr = TakesWithComponentsCandidates.Find(SlateName);
+	FTakeWithComponents* TakeWithComponentsPtr = TakesWithComponentsCandidates.Find(TakeIdentifier);
 	if (TakeWithComponentsPtr == nullptr)
 	{
-		TakeWithComponentsPtr = &TakesWithComponentsCandidates.Add(SlateName, { TakeWithComponents.TakePath, TakeWithComponents.SlateName, TakeWithComponents.TakeNumber, {} });
+		TakeWithComponentsPtr = &TakesWithComponentsCandidates.Add(TakeIdentifier, { TakeWithComponents.TakePath, TakeWithComponents.SlateName, TakeWithComponents.TakeNumber, {} });
 	}
 
 	for (FTakeWithComponents::Component& Component : TakeWithComponents.Components)
diff --git a/Engine/Plugins/VirtualProduction/CaptureManager/CaptureManagerDevices/Source/StereoVideoLiveLinkDevice/Private/StereoVideoLiveLinkDevice.h b/Engine/Plugins/VirtualProduction/CaptureManager/CaptureManagerDevices/Source/StereoVideoLiveLinkDevice/Private/StereoVideoLiveLinkDevice.h
--- a/Engine/Plugins/VirtualProduction/CaptureManager/CaptureManagerDevices/Source/StereoVideoLiveLinkDevice/Private/StereoVideoLiveLinkDevice.h
+++ b/Engine/Plugins/VirtualProduction/CaptureManager/CaptureManagerDevices/Source/StereoVideoLiveLinkDevice/Private/StereoVideoLiveLinkDevice.h
@@ -6,6 +6,7 @@
 
 #include "Ingest/LiveLinkDeviceCapability_Ingest.h"
 #include "LiveLinkDeviceCapability_Connection.h"
+#include "Customizations/TakeDiscoveryExpressionCustomization.h"
 
 #include "Misc/Optional.h"
 
@@ -22,6 +23,12 @@ class UStereoVideoLiveLinkDeviceSettings : public ULiveLinkDeviceSettings
 	GENERATED_BODY()
 
 public:
+	UStereoVideoLiveLinkDeviceSettings()
+	{
+		VideoDiscoveryExpression = TEXT("<Auto>");
+		AudioDiscoveryExpression = TEXT("<Auto>");
+	}
+
 	UPROPERTY(EditAnywhere, BlueprintReadWrite, Category = "Stereo Video Device")
 	FString DisplayName = TEXT("Stereo Video Device");
 
@@ -29,31 +36,53 @@ public:
 	UPROPERTY(EditAnywhere, BlueprintReadWrite, Category = "Stereo Video Device")
 	FDirectoryPath TakeDirectory;
 
-	/** 
+	/**
 	 * Format of the file name that the device expects
-	 *		Available tokens:
-	 *			* {SlateName}
-	 *			* {TakeNumber}
-	 *			* {Name}
-	 *			* {Any}
+	 *		Available expressions:
+	 *			* <Slate> - slate name
+	 *			* <Name> - identifier for the camera in the stereo pair
+	 *			* <Take> - take number
+	 *			* <Any> - any string
+	 *			* <Auto> - used alone to automatically figure out takes based on directory structure
 	 *		Available delimiters: '_', '-', '.', '\'
+	 *
+	 * Examples:
+	 * Considering input video MyTakeFolder/MySlate_MyName_SomeString-005.mov
+	 *
+	 * Video File Name Pattern <Auto> - will set Slate to the containing folder “MyTakeFolder”, Name will be set to the video filename “MySlate_MyName_SomeString-005” and Take will be set to 1 by default
+	 *
+	 * Video File Name Pattern <Slate>_<Name>_<Any>-<Take> - will extract Slate as “MySlate”, Name as “MyName” and Take as 5. <Any> ensures that the “SomeString” part of the filename is ignored.
+	 *
+	 * In both cases, extension '.mov' is stripped out.
 	 * 
+	 * Note: when not using <Auto> expression, then both <Slate> and <Name> must be present in the pattern
 	 */
 	UPROPERTY(EditAnywhere, BlueprintReadWrite, Category = "Stereo Video Device")
-	FString VideoFileNameFormat = TEXT("{Any}");
+	FTakeDiscoveryExpression VideoDiscoveryExpression;
 
 	/**
 	 * Format of the file name that the device expects
-	 *		Available tokens:
-	 *			* {SlateName}
-	 *			* {TakeNumber}
-	 *			* {Name}
-	 *			* {Any}
+	 *		Available expressions:
+	 *			* <Slate> - slate name
+	 *			* <Name> - identifier for the camera in the stereo pair
+	 *			* <Take> - take number
+	 *			* <Any> - any string
+	 *			* <Auto> - used alone to automatically figure out takes based on directory structure
 	 *		Available delimiters: '_', '-', '.', '\'
 	 *
+	 * Examples:
+	 * Considering input audio MyTakeFolder/MySlate_MyName_SomeString-005.wav
+	 * 
+	  * Audio File Name Pattern <Auto> - will set Slate to the containing folder “MyTakeFolder”, Name will be set to the audio filename “MySlate_MyName_SomeString-005” and Take will be set to 1 by default
+	 *
+	 * Audio File Name Pattern <Slate>_<Name>_<Any>-<Take> - will extract Slate as “MySlate”, Name as “MyName” and Take as 5. <Any> ensures that the “SomeString” part of the filename is ignored.
+	 *
+	 * In both cases, extension '.wav' is stripped out.
+	 * 
+	 * Note: when not using <Auto> expression, then both <Slate> and <Name> must be present in the pattern
 	 */
 	UPROPERTY(EditAnywhere, BlueprintReadWrite, Category = "Stereo Video Device")
-	FString AudioFileNameFormat = TEXT("{Any}");
+	FTakeDiscoveryExpression AudioDiscoveryExpression;
 };
 
 
diff --git a/Engine/Plugins/VirtualProduction/CaptureManager/CaptureManagerDevices/Source/StereoVideoLiveLinkDevice/Private/StereoVideoLiveLinkDeviceModule.cpp b/Engine/Plugins/VirtualProduction/CaptureManager/CaptureManagerDevices/Source/StereoVideoLiveLinkDevice/Private/StereoVideoLiveLinkDeviceModule.cpp
--- a/Engine/Plugins/VirtualProduction/CaptureManager/CaptureManagerDevices/Source/StereoVideoLiveLinkDevice/Private/StereoVideoLiveLinkDeviceModule.cpp
+++ b/Engine/Plugins/VirtualProduction/CaptureManager/CaptureManagerDevices/Source/StereoVideoLiveLinkDevice/Private/StereoVideoLiveLinkDeviceModule.cpp
@@ -2,4 +2,33 @@
 
 #include "Modules/ModuleManager.h"
 
-IMPLEMENT_MODULE(FDefaultModuleImpl, StereoVideoLiveLinkDevice)
\ No newline at end of file
+#include "PropertyEditorModule.h"
+#include "Customizations/TakeDiscoveryExpressionCustomization.h"
+#include "Templates/SharedPointer.h"
+
+class FStereoVideoLiveLinkDeviceModule : public IModuleInterface
+{
+public:
+
+	void StartupModule()
+	{
+		FPropertyEditorModule& PropertyModule = FModuleManager::LoadModuleChecked<FPropertyEditorModule>("PropertyEditor");
+		PropertyModule.RegisterCustomPropertyTypeLayout(FTakeDiscoveryExpression::StaticStruct()->GetFName(),
+			FOnGetPropertyTypeCustomizationInstance::CreateRaw(this, &FStereoVideoLiveLinkDeviceModule::GetTakeDiscoveryExpressionCustomization));
+	}
+
+	void ShutdownModule()
+	{
+		FPropertyEditorModule& PropertyModule = FModuleManager::LoadModuleChecked<FPropertyEditorModule>("PropertyEditor");
+		PropertyModule.UnregisterCustomPropertyTypeLayout(FTakeDiscoveryExpression::StaticStruct()->GetFName());
+	}
+
+private:
+
+	TSharedRef<IPropertyTypeCustomization> GetTakeDiscoveryExpressionCustomization()
+	{		
+		return MakeShared<FTakeDiscoveryExpressionCustomization>();
+	}
+};
+
+IMPLEMENT_MODULE(FStereoVideoLiveLinkDeviceModule, StereoVideoLiveLinkDevice)
\ No newline at end of file
diff --git a/Engine/Plugins/VirtualProduction/CaptureManager/CaptureManagerDevices/Source/VideoLiveLinkDeviceCommon/Private/Customizations/TakeDiscoveryExpressionCustomization.cpp b/Engine/Plugins/VirtualProduction/CaptureManager/CaptureManagerDevices/Source/VideoLiveLinkDeviceCommon/Private/Customizations/TakeDiscoveryExpressionCustomization.cpp
new file mode 100644
--- /dev/null
+++ b/Engine/Plugins/VirtualProduction/CaptureManager/CaptureManagerDevices/Source/VideoLiveLinkDeviceCommon/Private/Customizations/TakeDiscoveryExpressionCustomization.cpp
@@ -0,0 +1,150 @@
+// Copyright Epic Games, Inc. All Rights Reserved.
+
+#include "Customizations/TakeDiscoveryExpressionCustomization.h"
+#include "Widgets/Input/SEditableTextBox.h"
+#include "Editor.h"
+
+#define LOCTEXT_NAMESPACE "DiscoveryExpressionCustomization"
+
+FTakeDiscoveryExpressionCustomization::FTakeDiscoveryExpressionCustomization()
+{
+}
+
+void FTakeDiscoveryExpressionCustomization::CustomizeHeader(TSharedRef<IPropertyHandle> PropertyHandle, FDetailWidgetRow& HeaderRow, IPropertyTypeCustomizationUtils& CustomizationUtils)
+{
+}
+
+void FTakeDiscoveryExpressionCustomization::CustomizeChildren(TSharedRef<IPropertyHandle> PropertyHandle, IDetailChildrenBuilder& ChildBuilder, IPropertyTypeCustomizationUtils& CustomizationUtils)
+{
+	ExpressionProperty = PropertyHandle;
+
+	ChildBuilder.AddProperty(PropertyHandle)
+		.CustomWidget()
+		.NameContent()
+		[
+			PropertyHandle->CreatePropertyNameWidget()
+		]
+		.ValueContent()
+		[
+			SNew(SEditableTextBox)
+				.Text(this, &FTakeDiscoveryExpressionCustomization::OnGetExpressionValue)
+				.OnVerifyTextChanged(this, &FTakeDiscoveryExpressionCustomization::OnExpressionValidate)
+				.OnTextCommitted(this, &FTakeDiscoveryExpressionCustomization::OnExpressionCommited)
+				.SelectAllTextWhenFocused(true)
+				.RevertTextOnEscape(true)
+				.Font(IDetailLayoutBuilder::GetDetailFont())
+				.IsReadOnly(this, &FTakeDiscoveryExpressionCustomization::IsReadOnly)
+		];
+}
+
+TOptional<FText> FTakeDiscoveryExpressionCustomization::ValidateExpression(const FString& InExpression)
+{
+	if (InExpression.IsEmpty())
+	{
+		return LOCTEXT("Discovery_Expression_Empty", "Expression field must not be empty");
+	}
+
+	if (InExpression == "<Auto>")
+	{
+		return {};
+	}
+
+	TArray<FString> MandatoryExpressionParts = { "<Slate>", "<Name>" };
+
+	for (const FString& ExpressionPart : MandatoryExpressionParts)
+	{
+		int32 FoundIndex = InExpression.Find(ExpressionPart);
+		if (FoundIndex == INDEX_NONE)
+		{
+			return FText::Format(LOCTEXT("Discovery_Expression_Missing_Mandatory_Part", "Missing mandatory part: {0}"), FText::FromString(ExpressionPart));
+		}
+
+		FoundIndex = InExpression.Find(ExpressionPart, ESearchCase::CaseSensitive, ESearchDir::FromStart, FoundIndex + 1);
+		if (FoundIndex != INDEX_NONE)
+		{
+			return FText::Format(LOCTEXT("Discovery_Expression_Part_Repetition", "Multiple expression parts: {0}"), FText::FromString(ExpressionPart));
+		}
+	}
+
+	TArray<FString> OptionalExpressionParts = { "<Take>" };
+
+	for (const FString& ExpressionPart : OptionalExpressionParts)
+	{
+		int32 FoundIndex = InExpression.Find(ExpressionPart);
+		if (FoundIndex != INDEX_NONE)
+		{
+			FoundIndex = InExpression.Find(ExpressionPart, ESearchCase::CaseSensitive, ESearchDir::FromStart, FoundIndex + 1);
+
+			if (FoundIndex != INDEX_NONE)
+			{
+				return FText::Format(LOCTEXT("Discovery_Expression_Part_Repetition", "Multiple expression parts: {0}"), FText::FromString(ExpressionPart));
+			}
+		}
+	}
+
+	return {};
+}
+
+FText FTakeDiscoveryExpressionCustomization::OnGetExpressionValue() const
+{
+	TArray<void*> RawValue;
+	ExpressionProperty->AccessRawData(RawValue);
+
+	if (RawValue.Num())
+	{
+		const FTakeDiscoveryExpression* Expression = static_cast<const FTakeDiscoveryExpression*>(RawValue[0]);
+		return FText::FromString(Expression->Value);
+	}
+
+	return FText::GetEmpty();
+}
+
+bool FTakeDiscoveryExpressionCustomization::OnExpressionValidate(const FText& InText, FText& OutErrorText)
+{
+	FString Value = InText.ToString();
+
+	TOptional<FText> MaybeError = ValidateExpression(InText.ToString());
+
+	if (MaybeError.IsSet())
+	{
+		OutErrorText = MaybeError.GetValue();
+		return false;
+	}
+
+	return true;
+}
+
+void FTakeDiscoveryExpressionCustomization::OnExpressionCommited(const FText& InText, ETextCommit::Type CommitInfo)
+{
+	if (CommitInfo == ETextCommit::Type::OnCleared)
+	{
+		return;
+	}
+
+	TArray<void*> RawValue;
+	ExpressionProperty->AccessRawData(RawValue);
+
+	FString Value = InText.ToString();
+
+	if (RawValue.Num())
+	{
+		GEditor->BeginTransaction(FText::Format(LOCTEXT("DiscoveryExpression_SetProperty", "Edit {0}"), ExpressionProperty->GetPropertyDisplayName()));
+
+		ExpressionProperty->NotifyPreChange();
+
+		FTakeDiscoveryExpression* TakesDiscoveryTokesPtr = static_cast<FTakeDiscoveryExpression*>(RawValue[0]);
+		TakesDiscoveryTokesPtr->Value = MoveTemp(Value);
+
+		ExpressionProperty->NotifyPostChange(EPropertyChangeType::ValueSet);
+		ExpressionProperty->NotifyFinishedChangingProperties();
+
+		GEditor->EndTransaction();
+	}
+}
+
+bool FTakeDiscoveryExpressionCustomization::IsReadOnly() const
+{
+	return !ExpressionProperty->IsEditable();
+}
+
+#undef LOCTEXT_NAMESPACE
\ No newline at end of file
diff --git a/Engine/Plugins/VirtualProduction/CaptureManager/CaptureManagerDevices/Source/VideoLiveLinkDeviceCommon/Private/Tests/TestVideoDeviceDiscoveryExpressionParser.cpp b/Engine/Plugins/VirtualProduction/CaptureManager/CaptureManagerDevices/Source/VideoLiveLinkDeviceCommon/Private/Tests/TestVideoDeviceDiscoveryExpressionParser.cpp
new file mode 100644
--- /dev/null
+++ b/Engine/Plugins/VirtualProduction/CaptureManager/CaptureManagerDevices/Source/VideoLiveLinkDeviceCommon/Private/Tests/TestVideoDeviceDiscoveryExpressionParser.cpp
@@ -0,0 +1,446 @@
+// Copyright Epic Games, Inc. All Rights Reserved.
+
+#include "Utils/TakeDiscoveryExpressionParser.h"
+
+#include "Misc/AutomationTest.h"
+
+static const TArray<FString::ElementType> Delimiters =
+{
+	'-',
+	'_',
+	'.',
+	'/'
+};
+
+IMPLEMENT_SIMPLE_AUTOMATION_TEST(FTakeDiscoveryExpressionParserTest1, "CaptureManager.VideoLiveLinkDeviceCommon.TokenParser.Success",
+								 EAutomationTestFlags::EditorContext | EAutomationTestFlags::EngineFilter);
+
+bool FTakeDiscoveryExpressionParserTest1::RunTest(const FString& Parameters)
+{
+	FString InputFormat = TEXT("{SlateName}_{TakeNumber}_{Name}");
+
+	FString SlateName = TEXT("SlateName");
+	int32 TakeNumber = 10;
+	FString Name = TEXT("Name");
+
+	FString FormattedValue = FString::Format(TEXT("{0}_{1}_{2}"), { SlateName, TakeNumber, Name });
+
+	FTakeDiscoveryExpressionParser Parser(MoveTemp(InputFormat), MoveTemp(FormattedValue), Delimiters);
+	TestTrue(TEXT("Parse succeeded"), Parser.Parse());
+	TestEqual(TEXT("Slate Name matches"), SlateName, Parser.GetSlateName());
+	TestEqual(TEXT("Take Number matches"), TakeNumber, Parser.GetTakeNumber());
+	TestEqual(TEXT("Name matches"), Name, Parser.GetName());
+
+	return true;
+}
+
+IMPLEMENT_SIMPLE_AUTOMATION_TEST(FTakeDiscoveryExpressionParserTest2, "CaptureManager.VideoLiveLinkDeviceCommon.TokenParser.Fail",
+								 EAutomationTestFlags::EditorContext | EAutomationTestFlags::EngineFilter);
+
+bool FTakeDiscoveryExpressionParserTest2::RunTest(const FString& Parameters)
+{
+	FString InputFormat = TEXT("{TakeNumber}_{Name}");
+
+	FString SlateName = TEXT("SlateName");
+	int32 TakeNumber = 10;
+	FString Name = TEXT("Name");
+
+	FString FormattedValue = FString::Format(TEXT("{0}_{1}_{2}"), { SlateName, TakeNumber, Name });
+
+	FTakeDiscoveryExpressionParser Parser(MoveTemp(InputFormat), MoveTemp(FormattedValue), Delimiters);
+	TestFalse(TEXT("Parse failed"), Parser.Parse());
+
+	return true;
+}
+
+IMPLEMENT_SIMPLE_AUTOMATION_TEST(FTakeDiscoveryExpressionParserTest3, "CaptureManager.VideoLiveLinkDeviceCommon.TokenParser.Postfix.SuccessAny",
+								 EAutomationTestFlags::EditorContext | EAutomationTestFlags::EngineFilter);
+
+bool FTakeDiscoveryExpressionParserTest3::RunTest(const FString& Parameters)
+{
+	FString InputFormat = TEXT("{SlateName}_{Name}_{Any}");
+
+	FString SlateName = TEXT("SlateName");
+	FString Name = TEXT("Name");
+
+	FString FormattedValue = FString::Format(TEXT("{0}_{1}_TestingAnythingHere"), { SlateName, Name });
+
+	FTakeDiscoveryExpressionParser Parser(MoveTemp(InputFormat), MoveTemp(FormattedValue), Delimiters);
+	TestTrue(TEXT("Parse succeeded"), Parser.Parse());
+	TestEqual(TEXT("Slate Name matches"), SlateName, Parser.GetSlateName());
+	TestEqual(TEXT("Name matches"), Name, Parser.GetName());
+
+	return true;
+}
+
+IMPLEMENT_SIMPLE_AUTOMATION_TEST(FTakeDiscoveryExpressionParserTest4, "CaptureManager.VideoLiveLinkDeviceCommon.TokenParser.Postfix.FailAny",
+								 EAutomationTestFlags::EditorContext | EAutomationTestFlags::EngineFilter);
+
+bool FTakeDiscoveryExpressionParserTest4::RunTest(const FString& Parameters)
+{
+	FString InputFormat = TEXT("{SlateName}_{Name}_");
+
+	FString SlateName = TEXT("SlateName");
+	FString Name = TEXT("Name");
+
+	FString FormattedValue = FString::Format(TEXT("{0}_{1}_TestingAnythingHere"), { SlateName, Name });
+
+	FTakeDiscoveryExpressionParser Parser(MoveTemp(InputFormat), MoveTemp(FormattedValue), Delimiters);
+	TestFalse(TEXT("Parse failed"), Parser.Parse());
+
+	return true;
+}
+
+IMPLEMENT_SIMPLE_AUTOMATION_TEST(FTakeDiscoveryExpressionParserTest5, "CaptureManager.VideoLiveLinkDeviceCommon.TokenParser.Postfix.SuccessCustom",
+								 EAutomationTestFlags::EditorContext | EAutomationTestFlags::EngineFilter);
+
+bool FTakeDiscoveryExpressionParserTest5::RunTest(const FString& Parameters)
+{
+	FString CustomString = TEXT("CustomString");
+
+	FString InputFormat = FString::Format(TEXT("{SlateName}_{Name}_{0}"), { CustomString });
+
+	FString SlateName = TEXT("SlateName");
+	FString Name = TEXT("Name");
+
+	FString FormattedValue = FString::Format(TEXT("{0}_{1}_{2}"), { SlateName, Name, CustomString });
+
+	FTakeDiscoveryExpressionParser Parser(MoveTemp(InputFormat), MoveTemp(FormattedValue), Delimiters);
+	TestTrue(TEXT("Parse succeeded"), Parser.Parse());
+	TestEqual(TEXT("Slate Name matches"), SlateName, Parser.GetSlateName());
+	TestEqual(TEXT("Name matches"), Name, Parser.GetName());
+
+	return true;
+}
+
+IMPLEMENT_SIMPLE_AUTOMATION_TEST(FTakeDiscoveryExpressionParserTest6, "CaptureManager.VideoLiveLinkDeviceCommon.TokenParser.Postfix.FailCustom",
+								 EAutomationTestFlags::EditorContext | EAutomationTestFlags::EngineFilter);
+
+bool FTakeDiscoveryExpressionParserTest6::RunTest(const FString& Parameters)
+{
+	FString CustomString = TEXT("CustomString");
+
+	FString InputFormat = FString::Format(TEXT("{SlateName}_{Name}_{0}"), { CustomString });
+
+	FString SlateName = TEXT("SlateName");
+	FString Name = TEXT("Name");
+
+	FString InvalidCustomString = TEXT("InvalidCustomString");
+	FString FormattedValue = FString::Format(TEXT("{0}_{1}_{2}"), { SlateName, Name, InvalidCustomString });
+
+	FTakeDiscoveryExpressionParser Parser(MoveTemp(InputFormat), MoveTemp(FormattedValue), Delimiters);
+	TestFalse(TEXT("Parse fail"), Parser.Parse());
+
+	return true;
+}
+
+IMPLEMENT_SIMPLE_AUTOMATION_TEST(FTakeDiscoveryExpressionParserTest7, "CaptureManager.VideoLiveLinkDeviceCommon.TokenParser.Prefix.SuccessAny",
+								 EAutomationTestFlags::EditorContext | EAutomationTestFlags::EngineFilter);
+
+bool FTakeDiscoveryExpressionParserTest7::RunTest(const FString& Parameters)
+{
+	FString InputFormat = TEXT("{Any}_{SlateName}_{Name}");
+
+	FString SlateName = TEXT("SlateName");
+	FString Name = TEXT("Name");
+
+	FString FormattedValue = FString::Format(TEXT("TestingAnythingHere_{0}_{1}"), { SlateName, Name });
+
+	FTakeDiscoveryExpressionParser Parser(MoveTemp(InputFormat), MoveTemp(FormattedValue), Delimiters);
+	TestTrue(TEXT("Parse succeeded"), Parser.Parse());
+	TestEqual(TEXT("Slate Name matches"), SlateName, Parser.GetSlateName());
+	TestEqual(TEXT("Name matches"), Name, Parser.GetName());
+
+	return true;
+}
+
+IMPLEMENT_SIMPLE_AUTOMATION_TEST(FTakeDiscoveryExpressionParserTest8, "CaptureManager.VideoLiveLinkDeviceCommon.TokenParser.Prefix.FailAny",
+								 EAutomationTestFlags::EditorContext | EAutomationTestFlags::EngineFilter);
+
+bool FTakeDiscoveryExpressionParserTest8::RunTest(const FString& Parameters)
+{
+	FString InputFormat = TEXT("_{SlateName}_{Name}");
+
+	FString SlateName = TEXT("SlateName");
+	FString Name = TEXT("Name");
+
+	FString FormattedValue = FString::Format(TEXT("TestingAnythingHere_{0}_{1}"), { SlateName, Name });
+
+	FTakeDiscoveryExpressionParser Parser(MoveTemp(InputFormat), MoveTemp(FormattedValue), Delimiters);
+	TestFalse(TEXT("Parse failed"), Parser.Parse());
+
+	return true;
+}
+
+IMPLEMENT_SIMPLE_AUTOMATION_TEST(FTakeDiscoveryExpressionParserTest9, "CaptureManager.VideoLiveLinkDeviceCommon.TokenParser.Prefix.SuccessCustom",
+								 EAutomationTestFlags::EditorContext | EAutomationTestFlags::EngineFilter);
+
+bool FTakeDiscoveryExpressionParserTest9::RunTest(const FString& Parameters)
+{
+	FString CustomString = TEXT("CustomString");
+
+	FString InputFormat = FString::Format(TEXT("{0}_{SlateName}_{Name}"), { CustomString });
+
+	FString SlateName = TEXT("SlateName");
+	FString Name = TEXT("Name");
+
+	FString FormattedValue = FString::Format(TEXT("{0}_{1}_{2}"), { CustomString, SlateName, Name });
+
+	FTakeDiscoveryExpressionParser Parser(MoveTemp(InputFormat), MoveTemp(FormattedValue), Delimiters);
+	TestTrue(TEXT("Parse succeeded"), Parser.Parse());
+	TestEqual(TEXT("Slate Name matches"), SlateName, Parser.GetSlateName());
+	TestEqual(TEXT("Name matches"), Name, Parser.GetName());
+
+	return true;
+}
+
+IMPLEMENT_SIMPLE_AUTOMATION_TEST(FTakeDiscoveryExpressionParserTest10, "CaptureManager.VideoLiveLinkDeviceCommon.TokenParser.Prefix.FailCustom",
+								 EAutomationTestFlags::EditorContext | EAutomationTestFlags::EngineFilter);
+
+bool FTakeDiscoveryExpressionParserTest10::RunTest(const FString& Parameters)
+{
+	FString CustomString = TEXT("CustomString");
+
+	FString InputFormat = FString::Format(TEXT("{0}_{SlateName}_{Name}"), { CustomString });
+
+	FString SlateName = TEXT("SlateName");
+	FString Name = TEXT("Name");
+
+	FString InvalidCustomString = TEXT("InvalidCustomString");
+	FString FormattedValue = FString::Format(TEXT("{0}_{1}_{2}"), { InvalidCustomString, SlateName, Name });
+
+	FTakeDiscoveryExpressionParser Parser(MoveTemp(InputFormat), MoveTemp(FormattedValue), Delimiters);
+	TestFalse(TEXT("Parse fail"), Parser.Parse());
+
+	return true;
+}
+
+IMPLEMENT_SIMPLE_AUTOMATION_TEST(FTakeDiscoveryExpressionParserTest11, "CaptureManager.VideoLiveLinkDeviceCommon.TokenParser.Mid.SuccessAny",
+								 EAutomationTestFlags::EditorContext | EAutomationTestFlags::EngineFilter);
+
+bool FTakeDiscoveryExpressionParserTest11::RunTest(const FString& Parameters)
+{
+	FString InputFormat = TEXT("{SlateName}_{Any}_{Name}");
+
+	FString SlateName = TEXT("SlateName");
+	FString Name = TEXT("Name");
+
+	FString FormattedValue = FString::Format(TEXT("{0}_TestingAnythingHere_{1}"), { SlateName, Name });
+
+	FTakeDiscoveryExpressionParser Parser(MoveTemp(InputFormat), MoveTemp(FormattedValue), Delimiters);
+	TestTrue(TEXT("Parse succeeded"), Parser.Parse());
+	TestEqual(TEXT("Slate Name matches"), SlateName, Parser.GetSlateName());
+	TestEqual(TEXT("Name matches"), Name, Parser.GetName());
+
+	return true;
+}
+
+IMPLEMENT_SIMPLE_AUTOMATION_TEST(FTakeDiscoveryExpressionParserTest12, "CaptureManager.VideoLiveLinkDeviceCommon.TokenParser.Mid.FailAny",
+								 EAutomationTestFlags::EditorContext | EAutomationTestFlags::EngineFilter);
+
+bool FTakeDiscoveryExpressionParserTest12::RunTest(const FString& Parameters)
+{
+	FString InputFormat = TEXT("{SlateName}_{Name}");
+
+	FString SlateName = TEXT("SlateName");
+	FString Name = TEXT("Name");
+
+	FString FormattedValue = FString::Format(TEXT("{0}_TestingAnythingHere_{1}"), { SlateName, Name });
+
+	FTakeDiscoveryExpressionParser Parser(MoveTemp(InputFormat), MoveTemp(FormattedValue), Delimiters);
+	TestFalse(TEXT("Parse failed"), Parser.Parse());
+
+	return true;
+}
+
+IMPLEMENT_SIMPLE_AUTOMATION_TEST(FTakeDiscoveryExpressionParserTest13, "CaptureManager.VideoLiveLinkDeviceCommon.TokenParser.Mid.SuccessCustom",
+								 EAutomationTestFlags::EditorContext | EAutomationTestFlags::EngineFilter);
+
+bool FTakeDiscoveryExpressionParserTest13::RunTest(const FString& Parameters)
+{
+	FString CustomString = TEXT("CustomString");
+
+	FString InputFormat = FString::Format(TEXT("{SlateName}_{0}_{Name}"), { CustomString });
+
+	FString SlateName = TEXT("SlateName");
+	FString Name = TEXT("Name");
+
+	FString FormattedValue = FString::Format(TEXT("{0}_{1}_{2}"), { SlateName, CustomString, Name });
+
+	FTakeDiscoveryExpressionParser Parser(MoveTemp(InputFormat), MoveTemp(FormattedValue), Delimiters);
+	TestTrue(TEXT("Parse succeeded"), Parser.Parse());
+	TestEqual(TEXT("Slate Name matches"), SlateName, Parser.GetSlateName());
+	TestEqual(TEXT("Name matches"), Name, Parser.GetName());
+
+	return true;
+}
+
+IMPLEMENT_SIMPLE_AUTOMATION_TEST(FTakeDiscoveryExpressionParserTest14, "CaptureManager.VideoLiveLinkDeviceCommon.TokenParser.Mid.FailCustom",
+								 EAutomationTestFlags::EditorContext | EAutomationTestFlags::EngineFilter);
+
+bool FTakeDiscoveryExpressionParserTest14::RunTest(const FString& Parameters)
+{
+	FString CustomString = TEXT("CustomString");
+
+	FString InputFormat = FString::Format(TEXT("{SlateName}_{0}_{Name}"), { CustomString });
+
+	FString SlateName = TEXT("SlateName");
+	FString Name = TEXT("Name");
+
+	FString InvalidCustomString = TEXT("InvalidCustomString");
+	FString FormattedValue = FString::Format(TEXT("{0}_{1}_{2}"), { SlateName, InvalidCustomString, Name });
+
+	FTakeDiscoveryExpressionParser Parser(MoveTemp(InputFormat), MoveTemp(FormattedValue), Delimiters);
+	TestFalse(TEXT("Parse fail"), Parser.Parse());
+
+	return true;
+}
+
+IMPLEMENT_SIMPLE_AUTOMATION_TEST(FTakeDiscoveryExpressionParserTest15, "CaptureManager.VideoLiveLinkDeviceCommon.TokenParser.DifferentDelimiter1Success",
+								 EAutomationTestFlags::EditorContext | EAutomationTestFlags::EngineFilter);
+
+bool FTakeDiscoveryExpressionParserTest15::RunTest(const FString& Parameters)
+{
+	FString InputFormat = TEXT("{SlateName}_{TakeNumber}-{Name}");
+
+	FString SlateName = TEXT("SlateName");
+	int32 TakeNumber = 10;
+	FString Name = TEXT("Name");
+
+	FString FormattedValue = FString::Format(TEXT("{0}_{1}-{2}"), { SlateName, TakeNumber, Name });
+
+	FTakeDiscoveryExpressionParser Parser(MoveTemp(InputFormat), MoveTemp(FormattedValue), Delimiters);
+	TestTrue(TEXT("Parse succeeded"), Parser.Parse());
+
+	return true;
+}
+
+IMPLEMENT_SIMPLE_AUTOMATION_TEST(FTakeDiscoveryExpressionParserTest16, "CaptureManager.VideoLiveLinkDeviceCommon.TokenParser.DifferentDelimiterFailure",
+								 EAutomationTestFlags::EditorContext | EAutomationTestFlags::EngineFilter);
+
+bool FTakeDiscoveryExpressionParserTest16::RunTest(const FString& Parameters)
+{
+	FString InputFormat = TEXT("{SlateName}-{Name}");
+
+	FString SlateName = TEXT("SlateName");
+	int32 TakeNumber = 10;
+	FString Name = TEXT("Name");
+
+	FString FormattedValue = FString::Format(TEXT("{0}_{1}-{2}"), { SlateName, TakeNumber, Name });
+
+	FTakeDiscoveryExpressionParser Parser(MoveTemp(InputFormat), MoveTemp(FormattedValue), Delimiters);
+	TestFalse(TEXT("Parse failed"), Parser.Parse());
+	
+	return true;
+}
+
+IMPLEMENT_SIMPLE_AUTOMATION_TEST(FTakeDiscoveryExpressionParserTest17, "CaptureManager.VideoLiveLinkDeviceCommon.TokenParser.SingleTokenValue",
+								 EAutomationTestFlags::EditorContext | EAutomationTestFlags::EngineFilter);
+
+bool FTakeDiscoveryExpressionParserTest17::RunTest(const FString& Parameters)
+{
+	FString InputFormat = TEXT("{SlateName}");
+
+	FString SlateName = TEXT("SlateName");
+	int32 TakeNumber = 10;
+
+	FString FormattedValue = FString::Format(TEXT("{0}"), { SlateName });
+
+	FTakeDiscoveryExpressionParser Parser(MoveTemp(InputFormat), MoveTemp(FormattedValue), Delimiters);
+	TestTrue(TEXT("Parse succeeded"), Parser.Parse());
+	TestEqual(TEXT("Slate Name matches"), SlateName, Parser.GetSlateName());
+
+	return true;
+}
+
+IMPLEMENT_SIMPLE_AUTOMATION_TEST(FTakeDiscoveryExpressionParserTest18, "CaptureManager.VideoLiveLinkDeviceCommon.TokenParser.SingleValue",
+								 EAutomationTestFlags::EditorContext | EAutomationTestFlags::EngineFilter);
+
+bool FTakeDiscoveryExpressionParserTest18::RunTest(const FString& Parameters)
+{
+	FString InputFormat = TEXT("TestingString");
+	FString TestingString = TEXT("TestingString");
+
+	int32 NoTakeNumber = INDEX_NONE;
+
+	FTakeDiscoveryExpressionParser Parser(MoveTemp(InputFormat), MoveTemp(TestingString), Delimiters);
+	TestTrue(TEXT("Parse succeeded"), Parser.Parse());
+	TestTrue(TEXT("Slate Name empty"), Parser.GetSlateName().IsEmpty());
+	TestEqual(TEXT("Take Number invalid"), NoTakeNumber, Parser.GetTakeNumber());
+	TestTrue(TEXT("Name empty"), Parser.GetName().IsEmpty());
+
+	return true;
+}
+
+IMPLEMENT_SIMPLE_AUTOMATION_TEST(FTakeDiscoveryExpressionParserTest19, "CaptureManager.VideoLiveLinkDeviceCommon.TokenParser.EmptyString",
+								 EAutomationTestFlags::EditorContext | EAutomationTestFlags::EngineFilter);
+
+bool FTakeDiscoveryExpressionParserTest19::RunTest(const FString& Parameters)
+{
+	FString InputFormat = TEXT("");
+
+	FString SlateName = TEXT("SlateName");
+	int32 TakeNumber = 10;
+
+	FString FormattedValue = FString::Format(TEXT("{0}_{1}"), { SlateName, TakeNumber });
+
+	FTakeDiscoveryExpressionParser Parser(MoveTemp(InputFormat), MoveTemp(FormattedValue), Delimiters);
+	TestFalse(TEXT("Parse failed"), Parser.Parse());
+
+	return true;
+}
+
+IMPLEMENT_SIMPLE_AUTOMATION_TEST(FTakeDiscoveryExpressionParserTest20, "CaptureManager.VideoLiveLinkDeviceCommon.TokenParser.MultipleDelimiters",
+								 EAutomationTestFlags::EditorContext | EAutomationTestFlags::EngineFilter);
+
+bool FTakeDiscoveryExpressionParserTest20::RunTest(const FString& Parameters)
+{
+	FString InputFormat = TEXT("______");
+
+	FString SlateName = TEXT("SlateName");
+	int32 TakeNumber = 10;
+
+	FString FormattedValue = FString::Format(TEXT("{0}_{1}"), { SlateName, TakeNumber });
+
+	FTakeDiscoveryExpressionParser Parser(MoveTemp(InputFormat), MoveTemp(FormattedValue), Delimiters);
+	TestFalse(TEXT("Parse failed"), Parser.Parse());
+
+	return true;
+}
+
+IMPLEMENT_SIMPLE_AUTOMATION_TEST(FTakeDiscoveryExpressionParserTest21, "CaptureManager.VideoLiveLinkDeviceCommon.TokenParser.SingleDelimiter",
+								 EAutomationTestFlags::EditorContext | EAutomationTestFlags::EngineFilter);
+
+bool FTakeDiscoveryExpressionParserTest21::RunTest(const FString& Parameters)
+{
+	FString InputFormat = TEXT("_");
+
+	FString SlateName = TEXT("SlateName");
+	int32 TakeNumber = 10;
+
+	FString FormattedValue = FString::Format(TEXT("{0}_{1}"), { SlateName, TakeNumber });
+
+	FTakeDiscoveryExpressionParser Parser(MoveTemp(InputFormat), MoveTemp(FormattedValue), Delimiters);
+	TestFalse(TEXT("Parse failed"), Parser.Parse());
+
+	return true;
+}
+
+IMPLEMENT_SIMPLE_AUTOMATION_TEST(FTakeDiscoveryExpressionParserTest22, "CaptureManager.VideoLiveLinkDeviceCommon.TokenParser.DifferentDelimiter2Success",
+	EAutomationTestFlags::EditorContext | EAutomationTestFlags::EngineFilter);
+
+bool FTakeDiscoveryExpressionParserTest22::RunTest(const FString& Parameters)
+{
+	FString InputFormat = TEXT("{SlateName}/{TakeNumber}-{Name}");
+
+	FString SlateName = TEXT("SlateName");
+	int32 TakeNumber = 10;
+	FString Name = TEXT("Name");
+
+	FString FormattedValue = FString::Format(TEXT("{0}/{1}-{2}"), { SlateName, TakeNumber, Name });
+
+	FTakeDiscoveryExpressionParser Parser(MoveTemp(InputFormat), MoveTemp(FormattedValue), Delimiters);
+	TestTrue(TEXT("Parse succeeded"), Parser.Parse());
+
+	return true;
+}
diff --git a/Engine/Plugins/VirtualProduction/CaptureManager/CaptureManagerDevices/Source/VideoLiveLinkDeviceCommon/Private/Tests/TestVideoDeviceTokenParser.cpp b/Engine/Plugins/VirtualProduction/CaptureManager/CaptureManagerDevices/Source/VideoLiveLinkDeviceCommon/Private/Tests/TestVideoDeviceTokenParser.cpp
deleted file mode 100644
diff --git a/Engine/Plugins/VirtualProduction/CaptureManager/CaptureManagerDevices/Source/VideoLiveLinkDeviceCommon/Private/Utils/TakeDiscoveryExpressionParser.cpp b/Engine/Plugins/VirtualProduction/CaptureManager/CaptureManagerDevices/Source/VideoLiveLinkDeviceCommon/Private/Utils/TakeDiscoveryExpressionParser.cpp
new file mode 100644
--- /dev/null
+++ b/Engine/Plugins/VirtualProduction/CaptureManager/CaptureManagerDevices/Source/VideoLiveLinkDeviceCommon/Private/Utils/TakeDiscoveryExpressionParser.cpp
@@ -0,0 +1,157 @@
+// Copyright Epic Games, Inc. All Rights Reserved.
+
+#include "Utils/TakeDiscoveryExpressionParser.h"
+
+#include "Containers/Map.h"
+#include "Containers/Array.h"
+#include "Containers/StringView.h"
+
+const FString FTakeDiscoveryExpressionParser::SlateNameToken = TEXT("<Slate>");
+const FString FTakeDiscoveryExpressionParser::TakeNumberToken = TEXT("<Take>");
+const FString FTakeDiscoveryExpressionParser::NameToken = TEXT("<Name>");
+const FString FTakeDiscoveryExpressionParser::AnyToken = TEXT("<Any>");
+
+FTakeDiscoveryExpressionParser::FTakeDiscoveryExpressionParser(const FString& InFormat, const FString& InFormattedValue, const TArray<FString::ElementType>& InDelimiters)
+	: Format(InFormat)
+	, FormattedValue(InFormattedValue)
+	, AllowedDelimiters(InDelimiters)
+{
+}
+
+bool FTakeDiscoveryExpressionParser::Parse()
+{
+	// Figure out which delimiters user actually uses in specified format
+	TArray<FString::ElementType> FoundDelimiters;
+	for (FString::ElementType AllowedDelimiter : AllowedDelimiters)
+	{
+		FString AllowedDelimiterStr;
+		AllowedDelimiterStr.AppendChar(AllowedDelimiter);
+		if (Format.Contains(AllowedDelimiterStr))
+		{
+			FoundDelimiters.Add(AllowedDelimiter);
+		}
+	}
+	
+	TOptional<FComponentValueMap> ComponentValueMap = ParseWithDelimiters(FoundDelimiters);
+	
+	if (!ComponentValueMap.IsSet())
+	{
+		return false;
+	}
+
+	return DetermineValues(MoveTemp(ComponentValueMap.GetValue()));
+}
+
+TOptional<FTakeDiscoveryExpressionParser::FComponentValueMap> FTakeDiscoveryExpressionParser::ParseWithDelimiters(TArray<FString::ElementType> InDelimiters) const
+{
+	TArray<FString> FormatComponents = ParseWithDelimiters(Format, InDelimiters);
+
+	if (FormatComponents.IsEmpty())
+	{
+		return {};
+	}
+
+	TArray<FString> FileNameComponents = ParseWithDelimiters(FormattedValue, InDelimiters);
+
+	if (FileNameComponents.IsEmpty())
+	{
+		return {};
+	}
+
+	if (FormatComponents.Num() != FileNameComponents.Num())
+	{
+		return {};
+	}
+
+	FComponentValueMap ComponentValueMap;
+	for (int32 Index = 0; Index < FileNameComponents.Num(); ++Index)
+	{
+		ComponentValueMap.Add(FormatComponents[Index], FileNameComponents[Index]);
+	}
+
+	return ComponentValueMap;
+}
+
+TArray<FString> FTakeDiscoveryExpressionParser::ParseWithDelimiters(FString InString, const TArray<FString::ElementType>& InDelimiters, int32 InDelimiterIndex) const
+{
+	TArray<FString> Result;
+	if (InDelimiterIndex >= InDelimiters.Num())
+	{
+		Result.Add(InString);
+		return Result;
+	}
+
+	FString Delimiter;
+	Delimiter.AppendChar(InDelimiters[InDelimiterIndex]);
+
+	TArray<FString> HalfwaySplitComponents;
+
+	InString.ParseIntoArray(HalfwaySplitComponents, *Delimiter);
+
+	for (const FString& Component : HalfwaySplitComponents)
+	{
+		for (const FString& Part : ParseWithDelimiters(Component, InDelimiters, InDelimiterIndex + 1))
+		{
+			Result.Add(Part);
+		}
+	}
+	
+	return Result;
+}
+
+
+TOptional<FTakeDiscoveryExpressionParser::FComponentValueMap> FTakeDiscoveryExpressionParser::ParseWithoutDelimiters() const
+{
+	FComponentValueMap ComponentValueMap;
+	ComponentValueMap.Add(Format, FormattedValue);
+
+	return ComponentValueMap;
+}
+
+bool FTakeDiscoveryExpressionParser::DetermineValues(FComponentValueMap InComponentValueMap)
+{
+	if (const FString* MaybeSlateName = InComponentValueMap.Find(SlateNameToken))
+	{
+		SlateName = *MaybeSlateName;
+		InComponentValueMap.Remove(SlateNameToken);
+	}
+
+	if (const FString* MaybeTakeNumber = InComponentValueMap.Find(TakeNumberToken))
+	{
+		TakeNumber = FCString::Atoi(**MaybeTakeNumber);
+		InComponentValueMap.Remove(TakeNumberToken);
+	}
+
+	if (const FString* MaybeName = InComponentValueMap.Find(NameToken))
+	{
+		Name = *MaybeName;
+		InComponentValueMap.Remove(NameToken);
+	}
+
+	InComponentValueMap.Remove(AnyToken);
+
+	for (const TPair<FString, FString>& AdditionalComponentValue : InComponentValueMap)
+	{
+		if (AdditionalComponentValue.Key != AdditionalComponentValue.Value)
+		{
+			return false;
+		}
+	}
+
+	return true;
+}
+
+FString FTakeDiscoveryExpressionParser::GetSlateName() const
+{
+	return SlateName;
+}
+
+int32 FTakeDiscoveryExpressionParser::GetTakeNumber() const
+{
+	return TakeNumber;
+}
+
+FString FTakeDiscoveryExpressionParser::GetName() const
+{
+	return Name;
+}
diff --git a/Engine/Plugins/VirtualProduction/CaptureManager/CaptureManagerDevices/Source/VideoLiveLinkDeviceCommon/Private/Utils/VideoDeviceTokenParser.cpp b/Engine/Plugins/VirtualProduction/CaptureManager/CaptureManagerDevices/Source/VideoLiveLinkDeviceCommon/Private/Utils/VideoDeviceTokenParser.cpp
deleted file mode 100644
diff --git a/Engine/Plugins/VirtualProduction/CaptureManager/CaptureManagerDevices/Source/VideoLiveLinkDeviceCommon/Public/Customizations/TakeDiscoveryExpressionCustomization.h b/Engine/Plugins/VirtualProduction/CaptureManager/CaptureManagerDevices/Source/VideoLiveLinkDeviceCommon/Public/Customizations/TakeDiscoveryExpressionCustomization.h
new file mode 100644
--- /dev/null
+++ b/Engine/Plugins/VirtualProduction/CaptureManager/CaptureManagerDevices/Source/VideoLiveLinkDeviceCommon/Public/Customizations/TakeDiscoveryExpressionCustomization.h
@@ -0,0 +1,42 @@
+// Copyright Epic Games, Inc. All Rights Reserved.
+
+#pragma once
+
+#include "DetailLayoutBuilder.h"
+#include "DetailWidgetRow.h"
+#include "IDetailChildrenBuilder.h"
+
+#include "TakeDiscoveryExpressionCustomization.generated.h"
+
+USTRUCT(BlueprintType)
+struct VIDEOLIVELINKDEVICECOMMON_API FTakeDiscoveryExpression
+{
+	GENERATED_BODY()
+
+	UPROPERTY(EditAnywhere, BlueprintReadWrite, Category = "Discovert Expression Value")
+	FString Value;
+
+	FTakeDiscoveryExpression operator=(const FString& InValue)
+	{
+		Value = InValue;
+		return *this;
+	}
+};
+
+class VIDEOLIVELINKDEVICECOMMON_API FTakeDiscoveryExpressionCustomization : public IPropertyTypeCustomization
+{
+public:
+	FTakeDiscoveryExpressionCustomization();
+
+	virtual void CustomizeHeader(TSharedRef<IPropertyHandle> PropertyHandle, FDetailWidgetRow& HeaderRow, IPropertyTypeCustomizationUtils& CustomizationUtils) override;
+	virtual void CustomizeChildren(TSharedRef<IPropertyHandle> PropertyHandle, IDetailChildrenBuilder& ChildBuilder, IPropertyTypeCustomizationUtils& CustomizationUtils) override;
+
+private:
+	FText OnGetExpressionValue() const;
+	bool OnExpressionValidate(const FText& InText, FText& OutErrorText);
+	void OnExpressionCommited(const FText& InText, ETextCommit::Type CommitInfo);
+	bool IsReadOnly() const;
+	TOptional<FText> ValidateExpression(const FString& InTokens);
+
+	TSharedPtr<IPropertyHandle> ExpressionProperty;
+};
diff --git a/Engine/Plugins/VirtualProduction/CaptureManager/CaptureManagerDevices/Source/VideoLiveLinkDeviceCommon/Public/Utils/TakeDiscoveryExpressionParser.h b/Engine/Plugins/VirtualProduction/CaptureManager/CaptureManagerDevices/Source/VideoLiveLinkDeviceCommon/Public/Utils/TakeDiscoveryExpressionParser.h
new file mode 100644
--- /dev/null
+++ b/Engine/Plugins/VirtualProduction/CaptureManager/CaptureManagerDevices/Source/VideoLiveLinkDeviceCommon/Public/Utils/TakeDiscoveryExpressionParser.h
@@ -0,0 +1,41 @@
+// Copyright Epic Games, Inc. All Rights Reserved.
+
+#pragma once
+
+#include "Containers/UnrealString.h"
+
+class VIDEOLIVELINKDEVICECOMMON_API FTakeDiscoveryExpressionParser
+{
+public:
+
+	static const FString SlateNameToken;
+	static const FString TakeNumberToken;
+	static const FString NameToken;
+	static const FString AnyToken;
+
+	FTakeDiscoveryExpressionParser(const FString& InFormat, const FString& InFormattedValue, const TArray<FString::ElementType>& InDelimiters);
+
+	bool Parse();
+
+	FString GetSlateName() const;
+	int32 GetTakeNumber() const;
+	FString GetName() const;
+
+private:
+
+	using FComponentValueMap = TMap<FString, FString>;
+
+	TOptional<FComponentValueMap> ParseWithDelimiters(TArray<FString::ElementType> InDelimiters) const;
+	TOptional<FComponentValueMap> ParseWithoutDelimiters() const;
+	TArray<FString> ParseWithDelimiters(FString InString, const TArray<FString::ElementType>& InDelimiters, int32 InDelimiterIndex = 0) const;
+
+	bool DetermineValues(FComponentValueMap InComponentValueMap);
+
+	FString Format;
+	FString FormattedValue;
+	const TArray<FString::ElementType>& AllowedDelimiters;
+
+	FString SlateName;
+	int32 TakeNumber = INDEX_NONE;
+	FString Name;
+};
\ No newline at end of file
diff --git a/Engine/Plugins/VirtualProduction/CaptureManager/CaptureManagerDevices/Source/VideoLiveLinkDeviceCommon/Public/Utils/VideoDeviceTokenParser.h b/Engine/Plugins/VirtualProduction/CaptureManager/CaptureManagerDevices/Source/VideoLiveLinkDeviceCommon/Public/Utils/VideoDeviceTokenParser.h
deleted file mode 100644
diff --git a/Engine/Plugins/VirtualProduction/CaptureManager/CaptureManagerDevices/Source/VideoLiveLinkDeviceCommon/VideoLiveLinkDeviceCommon.Build.cs b/Engine/Plugins/VirtualProduction/CaptureManager/CaptureManagerDevices/Source/VideoLiveLinkDeviceCommon/VideoLiveLinkDeviceCommon.Build.cs
--- a/Engine/Plugins/VirtualProduction/CaptureManager/CaptureManagerDevices/Source/VideoLiveLinkDeviceCommon/VideoLiveLinkDeviceCommon.Build.cs
+++ b/Engine/Plugins/VirtualProduction/CaptureManager/CaptureManagerDevices/Source/VideoLiveLinkDeviceCommon/VideoLiveLinkDeviceCommon.Build.cs
@@ -18,7 +18,11 @@ public class VideoLiveLinkDeviceCommon : ModuleRules
 		PrivateDependencyModuleNames.AddRange(new string[]
 		{
 			"CaptureManagerMediaRW",
-			"Engine"
+			"Engine",
+			"Slate",
+			"SlateCore",
+			"CoreUObject",
+			"UnrealEd"
 		});
 
 		CppCompileWarningSettings.UnsafeTypeCastWarningLevel = WarningLevel.Error;
