[CaptureManager] Added tooltips to add live link device list

diff --git a/Engine/Plugins/Animation/LiveLinkDevice/Source/LiveLinkDevice/Private/LiveLinkDeviceModule.cpp b/Engine/Plugins/Animation/LiveLinkDevice/Source/LiveLinkDevice/Private/LiveLinkDeviceModule.cpp
--- a/Engine/Plugins/Animation/LiveLinkDevice/Source/LiveLinkDevice/Private/LiveLinkDeviceModule.cpp
+++ b/Engine/Plugins/Animation/LiveLinkDevice/Source/LiveLinkDevice/Private/LiveLinkDeviceModule.cpp
@@ -86,9 +86,15 @@ TSharedRef<SWidget> FLiveLinkDeviceModule::OnGenerateAddDeviceMenu()
 			continue;
 		}
 
+		FText ToolTip = DeviceClass->GetToolTipText();
+		if (ToolTip.IsEmpty())
+		{
+			ToolTip = FText::FromString(DeviceClass->GetPathName());
+		}
+
 		MenuBuilder.AddMenuEntry(
 			DeviceClass->GetDisplayNameText(),
-			FText::FromString(DeviceClass->GetPathName()),
+			ToolTip,
 			FSlateIcon(),
 			FUIAction(FExecuteAction::CreateLambda(
 				[DeviceSubsystem, DeviceClass]
diff --git a/Engine/Plugins/VirtualProduction/CaptureManager/CaptureManagerDevices/Source/CPSLiveLinkDevice/Private/LiveLinkFaceDevice.h b/Engine/Plugins/VirtualProduction/CaptureManager/CaptureManagerDevices/Source/CPSLiveLinkDevice/Private/LiveLinkFaceDevice.h
--- a/Engine/Plugins/VirtualProduction/CaptureManager/CaptureManagerDevices/Source/CPSLiveLinkDevice/Private/LiveLinkFaceDevice.h
+++ b/Engine/Plugins/VirtualProduction/CaptureManager/CaptureManagerDevices/Source/CPSLiveLinkDevice/Private/LiveLinkFaceDevice.h
@@ -26,21 +26,21 @@ class ULiveLinkFaceDeviceSettings : public ULiveLinkDeviceSettings
 	GENERATED_BODY()
 
 public:
-	UPROPERTY(EditAnywhere, BlueprintReadWrite, Category = "Live Link Face Device")
-	FString DisplayName = TEXT("Live Link Face Device");
+	UPROPERTY(EditAnywhere, BlueprintReadWrite, Category = "Live Link Face")
+	FString DisplayName = TEXT("Live Link Face");
 
-	UPROPERTY(EditAnywhere, BlueprintReadWrite, Category = "Live Link Face Device")
+	UPROPERTY(EditAnywhere, BlueprintReadWrite, Category = "Live Link Face")
 	FDeviceIpAddress IpAddress;
 
-	UPROPERTY(EditAnywhere, BlueprintReadWrite, Category = "Live Link Face Device")
+	UPROPERTY(EditAnywhere, BlueprintReadWrite, Category = "Live Link Face")
 	int32 Port = 14785;
 
-	UPROPERTY(VisibleAnywhere, Category = "Live Link Face Device")
+	UPROPERTY(VisibleAnywhere, Category = "Live Link Face")
 	FToggleConnectAction ConnectAction;
 };
 
 
-UCLASS(BlueprintType, meta = (DisplayName = "Live Link Face"))
+UCLASS(BlueprintType, meta = (DisplayName = "Live Link Face", ToolTip = "Use for ingest from the Live Link Face app and control recording"))
 class ULiveLinkFaceDevice final : public UBaseIngestLiveLinkDevice
 	, public ILiveLinkDeviceCapability_Connection
 	, public ILiveLinkDeviceCapability_Recording
diff --git a/Engine/Plugins/VirtualProduction/CaptureManager/CaptureManagerDevices/Source/FilesystemLiveLinkDevice/Private/FilesystemLiveLinkDevice.h b/Engine/Plugins/VirtualProduction/CaptureManager/CaptureManagerDevices/Source/FilesystemLiveLinkDevice/Private/FilesystemLiveLinkDevice.h
--- a/Engine/Plugins/VirtualProduction/CaptureManager/CaptureManagerDevices/Source/FilesystemLiveLinkDevice/Private/FilesystemLiveLinkDevice.h
+++ b/Engine/Plugins/VirtualProduction/CaptureManager/CaptureManagerDevices/Source/FilesystemLiveLinkDevice/Private/FilesystemLiveLinkDevice.h
@@ -22,16 +22,16 @@ class UFilesystemLiveLinkDeviceSettings : public ULiveLinkDeviceSettings
 	GENERATED_BODY()
 
 public:
-	UPROPERTY(EditAnywhere, BlueprintReadWrite, Category = "Filesystem Device")
-	FString DisplayName = TEXT("Filesystem Device");
+	UPROPERTY(EditAnywhere, BlueprintReadWrite, Category = "Take Archive Ingest")
+	FString DisplayName = TEXT("Take Archive Ingest");
 
 	/* Path to a directory containing the take(s) data */
-	UPROPERTY(EditAnywhere, BlueprintReadWrite, Category = "Filesystem Device")
+	UPROPERTY(EditAnywhere, BlueprintReadWrite, Category = "Take Archive Ingest")
 	FDirectoryPath TakeDirectory;
 };
 
 
-UCLASS(BlueprintType, meta=(DisplayName="Filesystem"))
+UCLASS(BlueprintType, meta = (DisplayName = "Take Archive Ingest", ToolTip = "Ingest take archives (described by .cptake files) and legacy Capture Manager takes"))
 class UFilesystemLiveLinkDevice final : public UBaseIngestLiveLinkDevice
 	, public ILiveLinkDeviceCapability_Connection
 {
diff --git a/Engine/Plugins/VirtualProduction/CaptureManager/CaptureManagerDevices/Source/MonoVideoLiveLinkDevice/Private/MonoVideoLiveLinkDevice.h b/Engine/Plugins/VirtualProduction/CaptureManager/CaptureManagerDevices/Source/MonoVideoLiveLinkDevice/Private/MonoVideoLiveLinkDevice.h
--- a/Engine/Plugins/VirtualProduction/CaptureManager/CaptureManagerDevices/Source/MonoVideoLiveLinkDevice/Private/MonoVideoLiveLinkDevice.h
+++ b/Engine/Plugins/VirtualProduction/CaptureManager/CaptureManagerDevices/Source/MonoVideoLiveLinkDevice/Private/MonoVideoLiveLinkDevice.h
@@ -28,11 +28,11 @@ public:
 		VideoDiscoveryExpression = TEXT("<Auto>");
 	}
 
-	UPROPERTY(EditAnywhere, BlueprintReadWrite, Category = "Mono Video Device")
-	FString DisplayName = TEXT("Mono Video Device");
+	UPROPERTY(EditAnywhere, BlueprintReadWrite, Category = "Mono Video Ingest")
+	FString DisplayName = TEXT("Mono Video Ingest");
 
 	/* Path to a directory containing the take(s) data */
-	UPROPERTY(EditAnywhere, BlueprintReadWrite, Category = "Mono Video Device")
+	UPROPERTY(EditAnywhere, BlueprintReadWrite, Category = "Mono Video Ingest")
 	FDirectoryPath TakeDirectory;
 
 	/**
@@ -56,12 +56,12 @@ public:
 	 *
 	 * Note: when not using <Auto> expression, then both <Slate> and <Name> must be present in the pattern.
 	 */
-	UPROPERTY(EditAnywhere, BlueprintReadWrite, Category = "Mono Video Device")
+	UPROPERTY(EditAnywhere, BlueprintReadWrite, Category = "Mono Video Ingest")
 	FTakeDiscoveryExpression VideoDiscoveryExpression;
 };
 
 
-UCLASS(BlueprintType, meta = (DisplayName = "Mono Video Device"))
+UCLASS(BlueprintType, meta = (DisplayName = "Mono Video Ingest", ToolTip = "Ingest video files as mono takes"))
 class UMonoVideoLiveLinkDevice final : public UBaseIngestLiveLinkDevice
 	, public ILiveLinkDeviceCapability_Connection
 {
diff --git a/Engine/Plugins/VirtualProduction/CaptureManager/CaptureManagerDevices/Source/StereoVideoLiveLinkDevice/Private/StereoVideoLiveLinkDevice.h b/Engine/Plugins/VirtualProduction/CaptureManager/CaptureManagerDevices/Source/StereoVideoLiveLinkDevice/Private/StereoVideoLiveLinkDevice.h
--- a/Engine/Plugins/VirtualProduction/CaptureManager/CaptureManagerDevices/Source/StereoVideoLiveLinkDevice/Private/StereoVideoLiveLinkDevice.h
+++ b/Engine/Plugins/VirtualProduction/CaptureManager/CaptureManagerDevices/Source/StereoVideoLiveLinkDevice/Private/StereoVideoLiveLinkDevice.h
@@ -29,11 +29,11 @@ public:
 		AudioDiscoveryExpression = TEXT("<Auto>");
 	}
 
-	UPROPERTY(EditAnywhere, BlueprintReadWrite, Category = "Stereo Video Device")
-	FString DisplayName = TEXT("Stereo Video Device");
+	UPROPERTY(EditAnywhere, BlueprintReadWrite, Category = "Stereo Video Ingest")
+	FString DisplayName = TEXT("Stereo Video Ingest");
 
 	/* Path to a directory containing the take(s) data */
-	UPROPERTY(EditAnywhere, BlueprintReadWrite, Category = "Stereo Video Device")
+	UPROPERTY(EditAnywhere, BlueprintReadWrite, Category = "Stereo Video Ingest")
 	FDirectoryPath TakeDirectory;
 
 	/**
@@ -57,7 +57,7 @@ public:
 	 * 
 	 * Note: when not using <Auto> expression, then both <Slate> and <Name> must be present in the pattern
 	 */
-	UPROPERTY(EditAnywhere, BlueprintReadWrite, Category = "Stereo Video Device")
+	UPROPERTY(EditAnywhere, BlueprintReadWrite, Category = "Stereo Video Ingest")
 	FTakeDiscoveryExpression VideoDiscoveryExpression;
 
 	/**
@@ -81,12 +81,12 @@ public:
 	 * 
 	 * Note: when not using <Auto> expression, then both <Slate> and <Name> must be present in the pattern
 	 */
-	UPROPERTY(EditAnywhere, BlueprintReadWrite, Category = "Stereo Video Device")
+	UPROPERTY(EditAnywhere, BlueprintReadWrite, Category = "Stereo Video Ingest")
 	FTakeDiscoveryExpression AudioDiscoveryExpression;
 };
 
 
-UCLASS(BlueprintType, DisplayName = "Stereo Video Device")
+UCLASS(BlueprintType, meta = (DisplayName = "Stereo Video Ingest", ToolTip = "Ingest subfolders containing pairs of video files as stereo takes"))
 class STEREOVIDEOLIVELINKDEVICE_API UStereoVideoLiveLinkDevice final : public UBaseIngestLiveLinkDevice
 	, public ILiveLinkDeviceCapability_Connection
 {
