Fix TakeDiscoveryExpressionParser tests

diff --git a/Engine/Plugins/VirtualProduction/CaptureManager/CaptureManagerDevices/Source/VideoLiveLinkDeviceCommon/Private/Tests/TestVideoDeviceDiscoveryExpressionParser.cpp b/Engine/Plugins/VirtualProduction/CaptureManager/CaptureManagerDevices/Source/VideoLiveLinkDeviceCommon/Private/Tests/TestVideoDeviceDiscoveryExpressionParser.cpp
--- a/Engine/Plugins/VirtualProduction/CaptureManager/CaptureManagerDevices/Source/VideoLiveLinkDeviceCommon/Private/Tests/TestVideoDeviceDiscoveryExpressionParser.cpp
+++ b/Engine/Plugins/VirtualProduction/CaptureManager/CaptureManagerDevices/Source/VideoLiveLinkDeviceCommon/Private/Tests/TestVideoDeviceDiscoveryExpressionParser.cpp
@@ -12,12 +12,12 @@ static const TArray<FString::ElementType> Delimiters =
 	'/'
 };
 
-IMPLEMENT_SIMPLE_AUTOMATION_TEST(FTakeDiscoveryExpressionParserTest1, "CaptureManager.VideoLiveLinkDeviceCommon.TokenParser.Success",
+IMPLEMENT_SIMPLE_AUTOMATION_TEST(FTakeDiscoveryExpressionParserTest1, "CaptureManager.VideoLiveLinkDeviceCommon.DiscoveryExpressionParser.Success",
 								 EAutomationTestFlags::EditorContext | EAutomationTestFlags::EngineFilter);
 
 bool FTakeDiscoveryExpressionParserTest1::RunTest(const FString& Parameters)
 {
-	FString InputFormat = TEXT("{SlateName}_{TakeNumber}_{Name}");
+	FString InputFormat = TEXT("<Slate>_<Take>_<Name>");
 
 	FString SlateName = TEXT("SlateName");
 	int32 TakeNumber = 10;
@@ -34,12 +34,12 @@ bool FTakeDiscoveryExpressionParserTest1::RunTest(const FString& Parameters)
 	return true;
 }
 
-IMPLEMENT_SIMPLE_AUTOMATION_TEST(FTakeDiscoveryExpressionParserTest2, "CaptureManager.VideoLiveLinkDeviceCommon.TokenParser.Fail",
+IMPLEMENT_SIMPLE_AUTOMATION_TEST(FTakeDiscoveryExpressionParserTest2, "CaptureManager.VideoLiveLinkDeviceCommon.DiscoveryExpressionParser.Fail",
 								 EAutomationTestFlags::EditorContext | EAutomationTestFlags::EngineFilter);
 
 bool FTakeDiscoveryExpressionParserTest2::RunTest(const FString& Parameters)
 {
-	FString InputFormat = TEXT("{TakeNumber}_{Name}");
+	FString InputFormat = TEXT("<Take>_<Name>");
 
 	FString SlateName = TEXT("SlateName");
 	int32 TakeNumber = 10;
@@ -53,12 +53,12 @@ bool FTakeDiscoveryExpressionParserTest2::RunTest(const FString& Parameters)
 	return true;
 }
 
-IMPLEMENT_SIMPLE_AUTOMATION_TEST(FTakeDiscoveryExpressionParserTest3, "CaptureManager.VideoLiveLinkDeviceCommon.TokenParser.Postfix.SuccessAny",
+IMPLEMENT_SIMPLE_AUTOMATION_TEST(FTakeDiscoveryExpressionParserTest3, "CaptureManager.VideoLiveLinkDeviceCommon.DiscoveryExpressionParser.Postfix.SuccessAny",
 								 EAutomationTestFlags::EditorContext | EAutomationTestFlags::EngineFilter);
 
 bool FTakeDiscoveryExpressionParserTest3::RunTest(const FString& Parameters)
 {
-	FString InputFormat = TEXT("{SlateName}_{Name}_{Any}");
+	FString InputFormat = TEXT("<Slate>_<Name>_<Any>");
 
 	FString SlateName = TEXT("SlateName");
 	FString Name = TEXT("Name");
@@ -73,12 +73,12 @@ bool FTakeDiscoveryExpressionParserTest3::RunTest(const FString& Parameters)
 	return true;
 }
 
-IMPLEMENT_SIMPLE_AUTOMATION_TEST(FTakeDiscoveryExpressionParserTest4, "CaptureManager.VideoLiveLinkDeviceCommon.TokenParser.Postfix.FailAny",
+IMPLEMENT_SIMPLE_AUTOMATION_TEST(FTakeDiscoveryExpressionParserTest4, "CaptureManager.VideoLiveLinkDeviceCommon.DiscoveryExpressionParser.Postfix.FailAny",
 								 EAutomationTestFlags::EditorContext | EAutomationTestFlags::EngineFilter);
 
 bool FTakeDiscoveryExpressionParserTest4::RunTest(const FString& Parameters)
 {
-	FString InputFormat = TEXT("{SlateName}_{Name}_");
+	FString InputFormat = TEXT("<Slate>_<Name>_");
 
 	FString SlateName = TEXT("SlateName");
 	FString Name = TEXT("Name");
@@ -91,14 +91,14 @@ bool FTakeDiscoveryExpressionParserTest4::RunTest(const FString& Parameters)
 	return true;
 }
 
-IMPLEMENT_SIMPLE_AUTOMATION_TEST(FTakeDiscoveryExpressionParserTest5, "CaptureManager.VideoLiveLinkDeviceCommon.TokenParser.Postfix.SuccessCustom",
+IMPLEMENT_SIMPLE_AUTOMATION_TEST(FTakeDiscoveryExpressionParserTest5, "CaptureManager.VideoLiveLinkDeviceCommon.DiscoveryExpressionParser.Postfix.SuccessCustom",
 								 EAutomationTestFlags::EditorContext | EAutomationTestFlags::EngineFilter);
 
 bool FTakeDiscoveryExpressionParserTest5::RunTest(const FString& Parameters)
 {
 	FString CustomString = TEXT("CustomString");
 
-	FString InputFormat = FString::Format(TEXT("{SlateName}_{Name}_{0}"), { CustomString });
+	FString InputFormat = FString::Format(TEXT("<Slate>_<Name>_{0}"), { CustomString });
 
 	FString SlateName = TEXT("SlateName");
 	FString Name = TEXT("Name");
@@ -113,14 +113,14 @@ bool FTakeDiscoveryExpressionParserTest5::RunTest(const FString& Parameters)
 	return true;
 }
 
-IMPLEMENT_SIMPLE_AUTOMATION_TEST(FTakeDiscoveryExpressionParserTest6, "CaptureManager.VideoLiveLinkDeviceCommon.TokenParser.Postfix.FailCustom",
+IMPLEMENT_SIMPLE_AUTOMATION_TEST(FTakeDiscoveryExpressionParserTest6, "CaptureManager.VideoLiveLinkDeviceCommon.DiscoveryExpressionParser.Postfix.FailCustom",
 								 EAutomationTestFlags::EditorContext | EAutomationTestFlags::EngineFilter);
 
 bool FTakeDiscoveryExpressionParserTest6::RunTest(const FString& Parameters)
 {
 	FString CustomString = TEXT("CustomString");
 
-	FString InputFormat = FString::Format(TEXT("{SlateName}_{Name}_{0}"), { CustomString });
+	FString InputFormat = FString::Format(TEXT("<Slate>_<Name>_{0}"), { CustomString });
 
 	FString SlateName = TEXT("SlateName");
 	FString Name = TEXT("Name");
@@ -134,12 +134,12 @@ bool FTakeDiscoveryExpressionParserTest6::RunTest(const FString& Parameters)
 	return true;
 }
 
-IMPLEMENT_SIMPLE_AUTOMATION_TEST(FTakeDiscoveryExpressionParserTest7, "CaptureManager.VideoLiveLinkDeviceCommon.TokenParser.Prefix.SuccessAny",
+IMPLEMENT_SIMPLE_AUTOMATION_TEST(FTakeDiscoveryExpressionParserTest7, "CaptureManager.VideoLiveLinkDeviceCommon.DiscoveryExpressionParser.Prefix.SuccessAny",
 								 EAutomationTestFlags::EditorContext | EAutomationTestFlags::EngineFilter);
 
 bool FTakeDiscoveryExpressionParserTest7::RunTest(const FString& Parameters)
 {
-	FString InputFormat = TEXT("{Any}_{SlateName}_{Name}");
+	FString InputFormat = TEXT("<Any>_<Slate>_<Name>");
 
 	FString SlateName = TEXT("SlateName");
 	FString Name = TEXT("Name");
@@ -154,12 +154,12 @@ bool FTakeDiscoveryExpressionParserTest7::RunTest(const FString& Parameters)
 	return true;
 }
 
-IMPLEMENT_SIMPLE_AUTOMATION_TEST(FTakeDiscoveryExpressionParserTest8, "CaptureManager.VideoLiveLinkDeviceCommon.TokenParser.Prefix.FailAny",
+IMPLEMENT_SIMPLE_AUTOMATION_TEST(FTakeDiscoveryExpressionParserTest8, "CaptureManager.VideoLiveLinkDeviceCommon.DiscoveryExpressionParser.Prefix.FailAny",
 								 EAutomationTestFlags::EditorContext | EAutomationTestFlags::EngineFilter);
 
 bool FTakeDiscoveryExpressionParserTest8::RunTest(const FString& Parameters)
 {
-	FString InputFormat = TEXT("_{SlateName}_{Name}");
+	FString InputFormat = TEXT("_<Slate>_<Name>");
 
 	FString SlateName = TEXT("SlateName");
 	FString Name = TEXT("Name");
@@ -172,14 +172,14 @@ bool FTakeDiscoveryExpressionParserTest8::RunTest(const FString& Parameters)
 	return true;
 }
 
-IMPLEMENT_SIMPLE_AUTOMATION_TEST(FTakeDiscoveryExpressionParserTest9, "CaptureManager.VideoLiveLinkDeviceCommon.TokenParser.Prefix.SuccessCustom",
+IMPLEMENT_SIMPLE_AUTOMATION_TEST(FTakeDiscoveryExpressionParserTest9, "CaptureManager.VideoLiveLinkDeviceCommon.DiscoveryExpressionParser.Prefix.SuccessCustom",
 								 EAutomationTestFlags::EditorContext | EAutomationTestFlags::EngineFilter);
 
 bool FTakeDiscoveryExpressionParserTest9::RunTest(const FString& Parameters)
 {
 	FString CustomString = TEXT("CustomString");
 
-	FString InputFormat = FString::Format(TEXT("{0}_{SlateName}_{Name}"), { CustomString });
+	FString InputFormat = FString::Format(TEXT("{0}_<Slate>_<Name>"), { CustomString });
 
 	FString SlateName = TEXT("SlateName");
 	FString Name = TEXT("Name");
@@ -194,14 +194,14 @@ bool FTakeDiscoveryExpressionParserTest9::RunTest(const FString& Parameters)
 	return true;
 }
 
-IMPLEMENT_SIMPLE_AUTOMATION_TEST(FTakeDiscoveryExpressionParserTest10, "CaptureManager.VideoLiveLinkDeviceCommon.TokenParser.Prefix.FailCustom",
+IMPLEMENT_SIMPLE_AUTOMATION_TEST(FTakeDiscoveryExpressionParserTest10, "CaptureManager.VideoLiveLinkDeviceCommon.DiscoveryExpressionParser.Prefix.FailCustom",
 								 EAutomationTestFlags::EditorContext | EAutomationTestFlags::EngineFilter);
 
 bool FTakeDiscoveryExpressionParserTest10::RunTest(const FString& Parameters)
 {
 	FString CustomString = TEXT("CustomString");
 
-	FString InputFormat = FString::Format(TEXT("{0}_{SlateName}_{Name}"), { CustomString });
+	FString InputFormat = FString::Format(TEXT("{0}_<Slate>_<Name>"), { CustomString });
 
 	FString SlateName = TEXT("SlateName");
 	FString Name = TEXT("Name");
@@ -215,12 +215,12 @@ bool FTakeDiscoveryExpressionParserTest10::RunTest(const FString& Parameters)
 	return true;
 }
 
-IMPLEMENT_SIMPLE_AUTOMATION_TEST(FTakeDiscoveryExpressionParserTest11, "CaptureManager.VideoLiveLinkDeviceCommon.TokenParser.Mid.SuccessAny",
+IMPLEMENT_SIMPLE_AUTOMATION_TEST(FTakeDiscoveryExpressionParserTest11, "CaptureManager.VideoLiveLinkDeviceCommon.DiscoveryExpressionParser.Mid.SuccessAny",
 								 EAutomationTestFlags::EditorContext | EAutomationTestFlags::EngineFilter);
 
 bool FTakeDiscoveryExpressionParserTest11::RunTest(const FString& Parameters)
 {
-	FString InputFormat = TEXT("{SlateName}_{Any}_{Name}");
+	FString InputFormat = TEXT("<Slate>_<Any>_<Name>");
 
 	FString SlateName = TEXT("SlateName");
 	FString Name = TEXT("Name");
@@ -235,12 +235,12 @@ bool FTakeDiscoveryExpressionParserTest11::RunTest(const FString& Parameters)
 	return true;
 }
 
-IMPLEMENT_SIMPLE_AUTOMATION_TEST(FTakeDiscoveryExpressionParserTest12, "CaptureManager.VideoLiveLinkDeviceCommon.TokenParser.Mid.FailAny",
+IMPLEMENT_SIMPLE_AUTOMATION_TEST(FTakeDiscoveryExpressionParserTest12, "CaptureManager.VideoLiveLinkDeviceCommon.DiscoveryExpressionParser.Mid.FailAny",
 								 EAutomationTestFlags::EditorContext | EAutomationTestFlags::EngineFilter);
 
 bool FTakeDiscoveryExpressionParserTest12::RunTest(const FString& Parameters)
 {
-	FString InputFormat = TEXT("{SlateName}_{Name}");
+	FString InputFormat = TEXT("<Slate>_<Name>");
 
 	FString SlateName = TEXT("SlateName");
 	FString Name = TEXT("Name");
@@ -253,14 +253,14 @@ bool FTakeDiscoveryExpressionParserTest12::RunTest(const FString& Parameters)
 	return true;
 }
 
-IMPLEMENT_SIMPLE_AUTOMATION_TEST(FTakeDiscoveryExpressionParserTest13, "CaptureManager.VideoLiveLinkDeviceCommon.TokenParser.Mid.SuccessCustom",
+IMPLEMENT_SIMPLE_AUTOMATION_TEST(FTakeDiscoveryExpressionParserTest13, "CaptureManager.VideoLiveLinkDeviceCommon.DiscoveryExpressionParser.Mid.SuccessCustom",
 								 EAutomationTestFlags::EditorContext | EAutomationTestFlags::EngineFilter);
 
 bool FTakeDiscoveryExpressionParserTest13::RunTest(const FString& Parameters)
 {
 	FString CustomString = TEXT("CustomString");
 
-	FString InputFormat = FString::Format(TEXT("{SlateName}_{0}_{Name}"), { CustomString });
+	FString InputFormat = FString::Format(TEXT("<Slate>_{0}_<Name>"), { CustomString });
 
 	FString SlateName = TEXT("SlateName");
 	FString Name = TEXT("Name");
@@ -275,14 +275,14 @@ bool FTakeDiscoveryExpressionParserTest13::RunTest(const FString& Parameters)
 	return true;
 }
 
-IMPLEMENT_SIMPLE_AUTOMATION_TEST(FTakeDiscoveryExpressionParserTest14, "CaptureManager.VideoLiveLinkDeviceCommon.TokenParser.Mid.FailCustom",
+IMPLEMENT_SIMPLE_AUTOMATION_TEST(FTakeDiscoveryExpressionParserTest14, "CaptureManager.VideoLiveLinkDeviceCommon.DiscoveryExpressionParser.Mid.FailCustom",
 								 EAutomationTestFlags::EditorContext | EAutomationTestFlags::EngineFilter);
 
 bool FTakeDiscoveryExpressionParserTest14::RunTest(const FString& Parameters)
 {
 	FString CustomString = TEXT("CustomString");
 
-	FString InputFormat = FString::Format(TEXT("{SlateName}_{0}_{Name}"), { CustomString });
+	FString InputFormat = FString::Format(TEXT("<Slate>_{0}_<Name>"), { CustomString });
 
 	FString SlateName = TEXT("SlateName");
 	FString Name = TEXT("Name");
@@ -296,12 +296,12 @@ bool FTakeDiscoveryExpressionParserTest14::RunTest(const FString& Parameters)
 	return true;
 }
 
-IMPLEMENT_SIMPLE_AUTOMATION_TEST(FTakeDiscoveryExpressionParserTest15, "CaptureManager.VideoLiveLinkDeviceCommon.TokenParser.DifferentDelimiter1Success",
+IMPLEMENT_SIMPLE_AUTOMATION_TEST(FTakeDiscoveryExpressionParserTest15, "CaptureManager.VideoLiveLinkDeviceCommon.DiscoveryExpressionParser.DifferentDelimiterSuccess",
 								 EAutomationTestFlags::EditorContext | EAutomationTestFlags::EngineFilter);
 
 bool FTakeDiscoveryExpressionParserTest15::RunTest(const FString& Parameters)
 {
-	FString InputFormat = TEXT("{SlateName}_{TakeNumber}-{Name}");
+	FString InputFormat = TEXT("<Slate>_<Take>-<Name>");
 
 	FString SlateName = TEXT("SlateName");
 	int32 TakeNumber = 10;
@@ -315,12 +315,12 @@ bool FTakeDiscoveryExpressionParserTest15::RunTest(const FString& Parameters)
 	return true;
 }
 
-IMPLEMENT_SIMPLE_AUTOMATION_TEST(FTakeDiscoveryExpressionParserTest16, "CaptureManager.VideoLiveLinkDeviceCommon.TokenParser.DifferentDelimiterFailure",
+IMPLEMENT_SIMPLE_AUTOMATION_TEST(FTakeDiscoveryExpressionParserTest16, "CaptureManager.VideoLiveLinkDeviceCommon.DiscoveryExpressionParser.DifferentDelimiterFailure",
 								 EAutomationTestFlags::EditorContext | EAutomationTestFlags::EngineFilter);
 
 bool FTakeDiscoveryExpressionParserTest16::RunTest(const FString& Parameters)
 {
-	FString InputFormat = TEXT("{SlateName}-{Name}");
+	FString InputFormat = TEXT("{Slate}-{Name}");
 
 	FString SlateName = TEXT("SlateName");
 	int32 TakeNumber = 10;
@@ -334,12 +334,12 @@ bool FTakeDiscoveryExpressionParserTest16::RunTest(const FString& Parameters)
 	return true;
 }
 
-IMPLEMENT_SIMPLE_AUTOMATION_TEST(FTakeDiscoveryExpressionParserTest17, "CaptureManager.VideoLiveLinkDeviceCommon.TokenParser.SingleTokenValue",
+IMPLEMENT_SIMPLE_AUTOMATION_TEST(FTakeDiscoveryExpressionParserTest17, "CaptureManager.VideoLiveLinkDeviceCommon.DiscoveryExpressionParser.SingleTokenValue",
 								 EAutomationTestFlags::EditorContext | EAutomationTestFlags::EngineFilter);
 
 bool FTakeDiscoveryExpressionParserTest17::RunTest(const FString& Parameters)
 {
-	FString InputFormat = TEXT("{SlateName}");
+	FString InputFormat = TEXT("<Slate>");
 
 	FString SlateName = TEXT("SlateName");
 	int32 TakeNumber = 10;
@@ -353,7 +353,7 @@ bool FTakeDiscoveryExpressionParserTest17::RunTest(const FString& Parameters)
 	return true;
 }
 
-IMPLEMENT_SIMPLE_AUTOMATION_TEST(FTakeDiscoveryExpressionParserTest18, "CaptureManager.VideoLiveLinkDeviceCommon.TokenParser.SingleValue",
+IMPLEMENT_SIMPLE_AUTOMATION_TEST(FTakeDiscoveryExpressionParserTest18, "CaptureManager.VideoLiveLinkDeviceCommon.DiscoveryExpressionParser.SingleValue",
 								 EAutomationTestFlags::EditorContext | EAutomationTestFlags::EngineFilter);
 
 bool FTakeDiscoveryExpressionParserTest18::RunTest(const FString& Parameters)
@@ -372,7 +372,7 @@ bool FTakeDiscoveryExpressionParserTest18::RunTest(const FString& Parameters)
 	return true;
 }
 
-IMPLEMENT_SIMPLE_AUTOMATION_TEST(FTakeDiscoveryExpressionParserTest19, "CaptureManager.VideoLiveLinkDeviceCommon.TokenParser.EmptyString",
+IMPLEMENT_SIMPLE_AUTOMATION_TEST(FTakeDiscoveryExpressionParserTest19, "CaptureManager.VideoLiveLinkDeviceCommon.DiscoveryExpressionParser.EmptyString",
 								 EAutomationTestFlags::EditorContext | EAutomationTestFlags::EngineFilter);
 
 bool FTakeDiscoveryExpressionParserTest19::RunTest(const FString& Parameters)
@@ -390,7 +390,7 @@ bool FTakeDiscoveryExpressionParserTest19::RunTest(const FString& Parameters)
 	return true;
 }
 
-IMPLEMENT_SIMPLE_AUTOMATION_TEST(FTakeDiscoveryExpressionParserTest20, "CaptureManager.VideoLiveLinkDeviceCommon.TokenParser.MultipleDelimiters",
+IMPLEMENT_SIMPLE_AUTOMATION_TEST(FTakeDiscoveryExpressionParserTest20, "CaptureManager.VideoLiveLinkDeviceCommon.DiscoveryExpressionParser.MultipleDelimiters",
 								 EAutomationTestFlags::EditorContext | EAutomationTestFlags::EngineFilter);
 
 bool FTakeDiscoveryExpressionParserTest20::RunTest(const FString& Parameters)
@@ -408,7 +408,7 @@ bool FTakeDiscoveryExpressionParserTest20::RunTest(const FString& Parameters)
 	return true;
 }
 
-IMPLEMENT_SIMPLE_AUTOMATION_TEST(FTakeDiscoveryExpressionParserTest21, "CaptureManager.VideoLiveLinkDeviceCommon.TokenParser.SingleDelimiter",
+IMPLEMENT_SIMPLE_AUTOMATION_TEST(FTakeDiscoveryExpressionParserTest21, "CaptureManager.VideoLiveLinkDeviceCommon.DiscoveryExpressionParser.SingleDelimiter",
 								 EAutomationTestFlags::EditorContext | EAutomationTestFlags::EngineFilter);
 
 bool FTakeDiscoveryExpressionParserTest21::RunTest(const FString& Parameters)
@@ -426,12 +426,12 @@ bool FTakeDiscoveryExpressionParserTest21::RunTest(const FString& Parameters)
 	return true;
 }
 
-IMPLEMENT_SIMPLE_AUTOMATION_TEST(FTakeDiscoveryExpressionParserTest22, "CaptureManager.VideoLiveLinkDeviceCommon.TokenParser.DifferentDelimiter2Success",
+IMPLEMENT_SIMPLE_AUTOMATION_TEST(FTakeDiscoveryExpressionParserTest22, "CaptureManager.VideoLiveLinkDeviceCommon.DiscoveryExpressionParser.DifferentDelimiter2Success",
 	EAutomationTestFlags::EditorContext | EAutomationTestFlags::EngineFilter);
 
 bool FTakeDiscoveryExpressionParserTest22::RunTest(const FString& Parameters)
 {
-	FString InputFormat = TEXT("{SlateName}/{TakeNumber}-{Name}");
+	FString InputFormat = TEXT("<Slate>/<Take>-<Name>");
 
 	FString SlateName = TEXT("SlateName");
 	int32 TakeNumber = 10;
