[CaptureManager] Fixed crash in ingest task progress

diff --git a/Engine/Plugins/VirtualProduction/CaptureManager/CaptureManagerApp/Source/CaptureDataConverter/Private/Nodes/CaptureConvertDepthData.cpp b/Engine/Plugins/VirtualProduction/CaptureManager/CaptureManagerApp/Source/CaptureDataConverter/Private/Nodes/CaptureConvertDepthData.cpp
--- a/Engine/Plugins/VirtualProduction/CaptureManager/CaptureManagerApp/Source/CaptureDataConverter/Private/Nodes/CaptureConvertDepthData.cpp
+++ b/Engine/Plugins/VirtualProduction/CaptureManager/CaptureManagerApp/Source/CaptureDataConverter/Private/Nodes/CaptureConvertDepthData.cpp
@@ -160,7 +160,7 @@ FCaptureConvertDepthData::FResult FCaptureConvertDepthData::CopyData()
 	}
 
 	FTaskProgress::FTask Task = Params.TaskProgress->StartTask();
-	FTaskProgress TaskProgress(Depth.FramesCount.GetValue(), FTaskProgress::FProgressReporter::CreateLambda([&Task](float InProgress)
+	TSharedPtr<FTaskProgress> TaskProgress = MakeShared<FTaskProgress>(Depth.FramesCount.GetValue(), FTaskProgress::FProgressReporter::CreateLambda([&Task](float InProgress)
 	{
 		Task.Update(InProgress);
 	}));
@@ -182,7 +182,7 @@ FCaptureConvertDepthData::FResult FCaptureConvertDepthData::CopyData()
 			return false;
 		}
 		
-		FTaskProgress::FTask Task = TaskProgress.StartTask();
+		FTaskProgress::FTask Task = TaskProgress->StartTask();
 
 		FCopyProgressReporter ProgressReporter(Task, Params.StopToken);
 
diff --git a/Engine/Plugins/VirtualProduction/CaptureManager/CaptureManagerApp/Source/CaptureDataConverter/Private/Nodes/CaptureConvertVideoData.cpp b/Engine/Plugins/VirtualProduction/CaptureManager/CaptureManagerApp/Source/CaptureDataConverter/Private/Nodes/CaptureConvertVideoData.cpp
--- a/Engine/Plugins/VirtualProduction/CaptureManager/CaptureManagerApp/Source/CaptureDataConverter/Private/Nodes/CaptureConvertVideoData.cpp
+++ b/Engine/Plugins/VirtualProduction/CaptureManager/CaptureManagerApp/Source/CaptureDataConverter/Private/Nodes/CaptureConvertVideoData.cpp
@@ -149,7 +149,7 @@ FCaptureConvertVideoData::FResult FCaptureConvertVideoData::CopyData()
 	}
 
 	FTaskProgress::FTask Task = Params.TaskProgress->StartTask();
-	FTaskProgress TaskProgress(Video.FramesCount.GetValue(), FTaskProgress::FProgressReporter::CreateLambda([&Task](float InProgress)
+	TSharedPtr<FTaskProgress> TaskProgress = MakeShared<FTaskProgress>(Video.FramesCount.GetValue(), FTaskProgress::FProgressReporter::CreateLambda([&Task](float InProgress)
 	{
 		Task.Update(InProgress);
 	}));
@@ -174,7 +174,7 @@ FCaptureConvertVideoData::FResult FCaptureConvertVideoData::CopyData()
 			return false;
 		}
 		
-		FTaskProgress::FTask Task = TaskProgress.StartTask();
+		FTaskProgress::FTask Task = TaskProgress->StartTask();
 
 		FCopyProgressReporter ProgressReporter(Task, Params.StopToken);
 
diff --git a/Engine/Plugins/VirtualProduction/CaptureManager/CaptureManagerApp/Source/IngestLiveLinkDevice/Private/Utils/CaptureExtractTimecode.cpp b/Engine/Plugins/VirtualProduction/CaptureManager/CaptureManagerApp/Source/IngestLiveLinkDevice/Private/Utils/CaptureExtractTimecode.cpp
--- a/Engine/Plugins/VirtualProduction/CaptureManager/CaptureManagerApp/Source/IngestLiveLinkDevice/Private/Utils/CaptureExtractTimecode.cpp
+++ b/Engine/Plugins/VirtualProduction/CaptureManager/CaptureManagerApp/Source/IngestLiveLinkDevice/Private/Utils/CaptureExtractTimecode.cpp
@@ -16,11 +16,12 @@
 DECLARE_LOG_CATEGORY_EXTERN(LogCaptureExtractTimecode, Log, All);
 DEFINE_LOG_CATEGORY(LogCaptureExtractTimecode);
 
-FCaptureExtractVideoTimecode::FCaptureExtractVideoTimecode(const FString& InFilePath)
+FCaptureExtractVideoTimecode::FCaptureExtractVideoTimecode(const FString& InFilePath) :
+	FilePath(InFilePath)
 {}
 
 UE::CaptureManager::FTimecodeAndRateResult FCaptureExtractVideoTimecode::Extract()
-{	
+{
 	using namespace UE::CaptureManager;
 
 	check(!FilePath.IsEmpty())
@@ -410,7 +411,7 @@ UE::CaptureManager::FTimecodeAndRateResult FCaptureExtractVideoTimecode::Extract
 					continue;
 				}
 			}
-			
+
 			if (!bFoundTimecode)
 			{
 				TArray<FString> TimecodePartsStr;
@@ -459,7 +460,8 @@ UE::CaptureManager::FTimecodeAndRateResult FCaptureExtractVideoTimecode::Extract
 	return Result;
 }
 
-FCaptureExtractAudioTimecode::FCaptureExtractAudioTimecode(const FString& InFilePath)
+FCaptureExtractAudioTimecode::FCaptureExtractAudioTimecode(const FString& InFilePath) :
+	FilePath(InFilePath)
 {}
 
 UE::CaptureManager::FTimecodeAndRateResult FCaptureExtractAudioTimecode::Extract()
diff --git a/Engine/Plugins/VirtualProduction/CaptureManager/CaptureManagerDevices/Source/FilesystemLiveLinkDevice/Private/FilesystemLiveLinkDevice.cpp b/Engine/Plugins/VirtualProduction/CaptureManager/CaptureManagerDevices/Source/FilesystemLiveLinkDevice/Private/FilesystemLiveLinkDevice.cpp
--- a/Engine/Plugins/VirtualProduction/CaptureManager/CaptureManagerDevices/Source/FilesystemLiveLinkDevice/Private/FilesystemLiveLinkDevice.cpp
+++ b/Engine/Plugins/VirtualProduction/CaptureManager/CaptureManagerDevices/Source/FilesystemLiveLinkDevice/Private/FilesystemLiveLinkDevice.cpp
@@ -251,14 +251,18 @@ TOptional<FTakeMetadata> UFilesystemLiveLinkDevice::ReadTake(const FString& InCu
 			Video.Path = FPaths::ConvertRelativePathToFull(CurrentDirectory, *Video.Path);
 		}
 
-		TSharedPtr<FCaptureExtractVideoTimecode> Extractor = MakeShareable(new FCaptureExtractVideoTimecode(Video.Path));
-		FTimecodeAndRateResult TimecodeAndRateResult = Extractor->Extract();
-		if (TimecodeAndRateResult.IsValid())
+		// Check if this is a file (and not an image sequence directory)
+		if (FPaths::FileExists(Video.Path))
 		{
-			FTimecodeAndRate TimecodeAndRate = TimecodeAndRateResult.GetValue();
-			if (!Video.TimecodeStart.IsSet())
+			TSharedPtr<FCaptureExtractVideoTimecode> Extractor = MakeShareable(new FCaptureExtractVideoTimecode(Video.Path));
+			FTimecodeAndRateResult TimecodeAndRateResult = Extractor->Extract();
+			if (TimecodeAndRateResult.IsValid())
 			{
-				Video.TimecodeStart = TimecodeAndRate.Timecode.ToString();
+				FTimecodeAndRate TimecodeAndRate = TimecodeAndRateResult.GetValue();
+				if (!Video.TimecodeStart.IsSet())
+				{
+					Video.TimecodeStart = TimecodeAndRate.Timecode.ToString();
+				}
 			}
 		}
 
@@ -267,24 +271,6 @@ TOptional<FTakeMetadata> UFilesystemLiveLinkDevice::ReadTake(const FString& InCu
 			VideoFrameRate = UE::CaptureManager::ParseFrameRate(Video.FrameRate);
 		}
 	}
-	for (FTakeMetadata::FVideo& Depth : TakeMetadata.Depth)
-	{
-		if (FPaths::IsRelative(*Depth.Path))
-		{
-			Depth.Path = FPaths::ConvertRelativePathToFull(CurrentDirectory, *Depth.Path);
-		}
-
-		TSharedPtr<FCaptureExtractVideoTimecode> Extractor = MakeShareable(new FCaptureExtractVideoTimecode(Depth.Path));
-		FTimecodeAndRateResult TimecodeAndRateResult = Extractor->Extract();
-		if (TimecodeAndRateResult.IsValid())
-		{
-			FTimecodeAndRate TimecodeAndRate = TimecodeAndRateResult.GetValue();
-			if (!Depth.TimecodeStart.IsSet())
-			{
-				Depth.TimecodeStart = TimecodeAndRate.Timecode.ToString();
-			}
-		}
-	}
 
 	for (FTakeMetadata::FAudio& Audio : TakeMetadata.Audio)
 	{
