[CaptureManager] Rename device classes

diff --git a/Engine/Plugins/VirtualProduction/CaptureManager/CaptureManagerDevices/CaptureManagerDevices.uplugin b/Engine/Plugins/VirtualProduction/CaptureManager/CaptureManagerDevices/CaptureManagerDevices.uplugin
--- a/Engine/Plugins/VirtualProduction/CaptureManager/CaptureManagerDevices/CaptureManagerDevices.uplugin
+++ b/Engine/Plugins/VirtualProduction/CaptureManager/CaptureManagerDevices/CaptureManagerDevices.uplugin
@@ -18,7 +18,7 @@
 	"Hidden": true,
 	"Modules": [
 		{
-			"Name": "MonoVideoLiveLinkDevice",
+			"Name": "MonoVideoIngestDevice",
 			"Type": "Editor",
 			"LoadingPhase": "Default"
 		},
@@ -28,12 +28,12 @@
 			"LoadingPhase": "Default"
 		},
 		{
-			"Name": "FilesystemLiveLinkDevice",
+			"Name": "TakeArchiveIngestDevice",
 			"Type": "Editor",
 			"LoadingPhase": "Default"
 		},
 		{
-			"Name": "StereoVideoLiveLinkDevice",
+			"Name": "StereoVideoIngestDevice",
 			"Type": "Editor",
 			"LoadingPhase": "Default"
 		},
diff --git a/Engine/Plugins/VirtualProduction/CaptureManager/CaptureManagerDevices/Source/FilesystemLiveLinkDevice/FilesystemLiveLinkDevice.Build.cs b/Engine/Plugins/VirtualProduction/CaptureManager/CaptureManagerDevices/Source/FilesystemLiveLinkDevice/FilesystemLiveLinkDevice.Build.cs
deleted file mode 100644
diff --git a/Engine/Plugins/VirtualProduction/CaptureManager/CaptureManagerDevices/Source/FilesystemLiveLinkDevice/Private/FilesystemLiveLinkDevice.cpp b/Engine/Plugins/VirtualProduction/CaptureManager/CaptureManagerDevices/Source/FilesystemLiveLinkDevice/Private/FilesystemLiveLinkDevice.cpp
deleted file mode 100644
diff --git a/Engine/Plugins/VirtualProduction/CaptureManager/CaptureManagerDevices/Source/FilesystemLiveLinkDevice/Private/FilesystemLiveLinkDevice.h b/Engine/Plugins/VirtualProduction/CaptureManager/CaptureManagerDevices/Source/FilesystemLiveLinkDevice/Private/FilesystemLiveLinkDevice.h
deleted file mode 100644
diff --git a/Engine/Plugins/VirtualProduction/CaptureManager/CaptureManagerDevices/Source/FilesystemLiveLinkDevice/Private/FilesystemLiveLinkDeviceModule.cpp b/Engine/Plugins/VirtualProduction/CaptureManager/CaptureManagerDevices/Source/FilesystemLiveLinkDevice/Private/FilesystemLiveLinkDeviceModule.cpp
deleted file mode 100644
diff --git a/Engine/Plugins/VirtualProduction/CaptureManager/CaptureManagerDevices/Source/FilesystemLiveLinkDevice/Private/Utils/GenericLiveLinkDeviceLog.h b/Engine/Plugins/VirtualProduction/CaptureManager/CaptureManagerDevices/Source/FilesystemLiveLinkDevice/Private/Utils/GenericLiveLinkDeviceLog.h
deleted file mode 100644
diff --git a/Engine/Plugins/VirtualProduction/CaptureManager/CaptureManagerDevices/Source/MonoVideoIngestDevice/MonoVideoIngestDevice.Build.cs b/Engine/Plugins/VirtualProduction/CaptureManager/CaptureManagerDevices/Source/MonoVideoIngestDevice/MonoVideoIngestDevice.Build.cs
new file mode 100644
--- /dev/null
+++ b/Engine/Plugins/VirtualProduction/CaptureManager/CaptureManagerDevices/Source/MonoVideoIngestDevice/MonoVideoIngestDevice.Build.cs
@@ -0,0 +1,42 @@
+// Copyright Epic Games, Inc. All Rights Reserved.
+
+using UnrealBuildTool;
+using System.IO;
+
+public class MonoVideoIngestDevice : ModuleRules
+{
+	public MonoVideoIngestDevice(ReadOnlyTargetRules Target) : base(Target)
+	{
+		PCHUsage = PCHUsageMode.UseExplicitOrSharedPCHs;
+		ShortName = "MonoVideoIngest";
+
+		PublicDependencyModuleNames.AddRange(new string[]
+		{
+			"CaptureUtils",
+			"Core"
+		});
+
+		PrivateDependencyModuleNames.AddRange(new string[]
+		{
+			"Engine",
+			"Json",
+			"CaptureManagerMediaRW",
+			"IngestLiveLinkDevice",
+			"LiveLinkDevice",
+			"LiveLinkCapabilities",
+			"Media",
+			"LiveLinkHubCaptureMessaging",
+			"LiveLinkHub",
+			"ToolWidgets",
+			"Slate",
+			"SlateCore",
+			"CaptureManagerTakeMetadata",
+			"CaptureManagerSettings",
+			"CoreUObject",
+			"UnrealEd",
+			"VideoLiveLinkDeviceCommon"
+		});
+
+		CppCompileWarningSettings.UnsafeTypeCastWarningLevel = WarningLevel.Error;
+	}
+}
diff --git a/Engine/Plugins/VirtualProduction/CaptureManager/CaptureManagerDevices/Source/MonoVideoIngestDevice/Private/MonoVideoIngestDevice.cpp b/Engine/Plugins/VirtualProduction/CaptureManager/CaptureManagerDevices/Source/MonoVideoIngestDevice/Private/MonoVideoIngestDevice.cpp
new file mode 100644
--- /dev/null
+++ b/Engine/Plugins/VirtualProduction/CaptureManager/CaptureManagerDevices/Source/MonoVideoIngestDevice/Private/MonoVideoIngestDevice.cpp
@@ -0,0 +1,341 @@
+// Copyright Epic Games, Inc. All Rights Reserved.
+
+#include "MonoVideoIngestDevice.h"
+
+#include "HAL/FileManager.h"
+#include "Settings/CaptureManagerSettings.h"
+#include "ImageUtils.h"
+
+#include "CaptureManagerMediaRWModule.h"
+
+#include "Misc/MessageDialog.h"
+#include "Utils/TakeDiscoveryExpressionParser.h"
+#include "Utils/VideoDeviceThumbnailExtractor.h"
+#include "Utils/CaptureExtractTimecode.h"
+
+#include "VideoLiveLinkDeviceLog.h"
+
+#define LOCTEXT_NAMESPACE "MonoVideoDevice"
+
+DEFINE_LOG_CATEGORY(LogVideoLiveLinkDevice);
+
+namespace UE::MonoVideoLiveLinkDevice::Private
+{
+
+static const TArray<FString::ElementType> Delimiters =
+{
+	'-',
+	'_',
+	'.'
+};
+
+static const TArray<FStringView> SupportedVideoFormats =
+{
+	TEXTVIEW("mp4"),
+	TEXTVIEW("mov")
+};
+
+
+FMediaRWManager& GetMediaRWManager()
+{
+	return FModuleManager::LoadModuleChecked<FCaptureManagerMediaRWModule>("CaptureManagerMediaRW").Get();
+}
+
+}
+
+const UMonoVideoIngestDeviceSettings* UMonoVideoIngestDevice::GetSettings() const
+{
+	return GetDeviceSettings<UMonoVideoIngestDeviceSettings>();
+}
+
+TSubclassOf<ULiveLinkDeviceSettings> UMonoVideoIngestDevice::GetSettingsClass() const
+{
+	return UMonoVideoIngestDeviceSettings::StaticClass();
+}
+
+FText UMonoVideoIngestDevice::GetDisplayName() const
+{
+	return FText::FromString(GetSettings()->DisplayName);
+}
+
+EDeviceHealth UMonoVideoIngestDevice::GetDeviceHealth() const
+{
+	return EDeviceHealth::Nominal;
+}
+
+FText UMonoVideoIngestDevice::GetHealthText() const
+{
+	return FText::FromString("Nominal");
+}
+
+void UMonoVideoIngestDevice::OnSettingChanged(const FPropertyChangedEvent& InPropertyChangedEvent)
+{
+	Super::OnSettingChanged(InPropertyChangedEvent);
+
+	const FName& PropertyName = InPropertyChangedEvent.GetMemberPropertyName();
+
+	if (GET_MEMBER_NAME_CHECKED(UMonoVideoIngestDeviceSettings, TakeDirectory) == PropertyName
+		|| GET_MEMBER_NAME_CHECKED(UMonoVideoIngestDeviceSettings, VideoDiscoveryExpression) == PropertyName)
+	{
+		SetConnectionStatus(ELiveLinkDeviceConnectionStatus::Disconnected);
+
+		FString Path = GetSettings()->TakeDirectory.Path;
+		if (!Path.IsEmpty() && FPaths::DirectoryExists(Path))
+		{
+			SetConnectionStatus(ELiveLinkDeviceConnectionStatus::Connected);
+		}
+		else
+		{
+			SetConnectionStatus(ELiveLinkDeviceConnectionStatus::Disconnected);
+		}
+	}
+}
+
+FString UMonoVideoIngestDevice::GetFullTakePath(UE::CaptureManager::FTakeId InTakeId) const
+{
+	if (const FString* TakePath = FullTakePaths.Find(InTakeId))
+	{
+		return *TakePath;
+	}
+
+	return FString();
+}
+
+void UMonoVideoIngestDevice::RunConvertAndUploadTake(const UIngestCapability_ProcessHandle* InProcessHandle, const UIngestCapability_Options* InIngestOptions)
+{
+	AsyncTask(ENamedThreads::AnyThread, [this,
+			  ProcessHandle = TStrongObjectPtr<const UIngestCapability_ProcessHandle>(InProcessHandle),
+			  IngestOptions = TStrongObjectPtr<const UIngestCapability_Options>(InIngestOptions)]()
+	{
+		static constexpr uint32 NumberOfTasks = 2; // Convert, Upload
+
+		using namespace UE::CaptureManager;
+		TSharedPtr<FTaskProgress> TaskProgress
+			= MakeShared<FTaskProgress>(NumberOfTasks, FTaskProgress::FProgressReporter::CreateLambda([this, ProcessHandle](double InProgress)
+		{
+			ExecuteProcessProgressReporter(ProcessHandle.Get(), InProgress);
+		}));
+
+		Super::IngestTake(ProcessHandle.Get(), IngestOptions.Get(), MoveTemp(TaskProgress));
+	});
+}
+
+void UMonoVideoIngestDevice::UpdateTakeList_Implementation(UIngestCapability_UpdateTakeListCallback* InCallback)
+{
+	AsyncTask(ENamedThreads::AnyThread, [this, Callback = TStrongObjectPtr<UIngestCapability_UpdateTakeListCallback>(InCallback)]
+	{
+		RemoveAllTakes();
+
+		FString StoragePath = GetSettings()->TakeDirectory.Path;
+		int32 DirectoriesInterrogatedCount = 0;
+		const int32 DirectoriesToInterrogateInOneRun = 200;
+	
+		TArray<FString> SupportedVideoFiles;
+		const bool bIterationResult = IFileManager::Get().IterateDirectoryRecursively(*StoragePath, [&](const TCHAR* InFileNameOrDirectory, bool bInIsDirectory)
+		{
+			if (bInIsDirectory)
+			{
+				if (++DirectoriesInterrogatedCount > DirectoriesToInterrogateInOneRun)
+				{
+					return false;
+				}
+			}
+
+			if (!bInIsDirectory && IsVideoFile(InFileNameOrDirectory))
+			{			
+				SupportedVideoFiles.Add(InFileNameOrDirectory);
+			}
+
+			return true;
+		});
+
+		for (const FString& SupportedVideoFile : SupportedVideoFiles)
+		{
+			TOptional<FTakeMetadata> TakeInfoResult = ReadTake(SupportedVideoFile);
+
+			if (TakeInfoResult.IsSet())
+			{
+				int32 TakeId = AddTake(TakeInfoResult.GetValue());
+				FullTakePaths.Add(TakeId, SupportedVideoFile);
+
+				PublishEvent<FIngestCapability_TakeAddedEvent>(TakeId);
+			}
+		}
+
+		ExecuteUpdateTakeListCallback(Callback.Get(), Execute_GetTakeIdentifiers(this));
+	});
+}
+
+ELiveLinkDeviceConnectionStatus UMonoVideoIngestDevice::GetConnectionStatus_Implementation() const
+{
+	const UMonoVideoIngestDeviceSettings* DeviceSettings = GetSettings();
+	FString Path = DeviceSettings->TakeDirectory.Path;
+
+	if (!Path.IsEmpty() && FPaths::DirectoryExists(Path))
+	{
+		return ELiveLinkDeviceConnectionStatus::Connected;
+	}
+
+	return ELiveLinkDeviceConnectionStatus::Disconnected;
+}
+
+FString UMonoVideoIngestDevice::GetHardwareId_Implementation() const
+{
+	return FPlatformMisc::GetDeviceId();
+}
+
+bool UMonoVideoIngestDevice::SetHardwareId_Implementation(const FString& HardwareID)
+{
+	return false;
+}
+
+bool UMonoVideoIngestDevice::Connect_Implementation()
+{
+	const UMonoVideoIngestDeviceSettings* DeviceSettings = GetSettings();
+	FString Path = DeviceSettings->TakeDirectory.Path;
+
+	if (!Path.IsEmpty() && FPaths::DirectoryExists(Path))
+	{
+		SetConnectionStatus(ELiveLinkDeviceConnectionStatus::Connected);
+		return true;
+	}
+
+	return false;
+}
+
+bool UMonoVideoIngestDevice::Disconnect_Implementation()
+{
+	SetConnectionStatus(ELiveLinkDeviceConnectionStatus::Disconnected);
+	return true;
+}
+
+bool UMonoVideoIngestDevice::IsVideoFile(const FString& InFileNameWithExtension)
+{
+	FString Extension = FPaths::GetExtension(InFileNameWithExtension);
+
+	return UE::MonoVideoLiveLinkDevice::Private::SupportedVideoFormats.Contains(Extension);
+}
+
+TOptional<FTakeMetadata> UMonoVideoIngestDevice::ReadTake(const FString& InCurrentTakeFile) const
+{
+	using namespace UE::CaptureManager;
+
+	FMediaRWManager& MediaManager = UE::MonoVideoLiveLinkDevice::Private::GetMediaRWManager();
+
+	FTakeMetadata TakeMetadata;
+
+	FFileStatData FileData = IFileManager::Get().GetStatData(*InCurrentTakeFile);
+
+	FString FileNameFormat = GetSettings()->VideoDiscoveryExpression.Value;
+	FString FileName = FPaths::GetBaseFilename(InCurrentTakeFile);
+
+	FString SlateName;
+	FString Name;
+	int32 TakeNumber = INDEX_NONE;
+
+	if (FileNameFormat != "<Auto>")
+	{
+		FTakeDiscoveryExpressionParser TokenParser(FileNameFormat, FileName, UE::MonoVideoLiveLinkDevice::Private::Delimiters);
+		bool bParse = TokenParser.Parse();
+
+		if (!bParse)
+		{
+			UE_LOG(LogVideoLiveLinkDevice, Warning, TEXT("Failed to match the specified format (%s) with the video file (%s)"), *(FileNameFormat), *InCurrentTakeFile);
+			return {};
+		}
+
+		SlateName = TokenParser.GetSlateName();
+		TakeNumber = TokenParser.GetTakeNumber();
+		Name = TokenParser.GetName();
+	}
+
+	if (SlateName.IsEmpty())
+	{
+		SlateName = MoveTemp(FileName);
+	}
+
+	if (TakeNumber == INDEX_NONE)
+	{
+		TakeNumber = 1;
+	}
+
+	if (Name.IsEmpty())
+	{
+		Name = TEXT("video");
+	}
+
+	TakeMetadata.Version.Major = 4;
+	TakeMetadata.Version.Minor = 1;
+
+	TakeMetadata.Slate = MoveTemp(SlateName);
+	TakeMetadata.TakeNumber = TakeNumber;
+
+	// Create new guid
+	TakeMetadata.UniqueId = FGuid::NewGuid().ToString(EGuidFormats::DigitsWithHyphensLower);
+	TakeMetadata.DateTime = FileData.CreationTime;
+	TakeMetadata.Device.Model = TEXT("MonoVideo");
+	
+	FTakeMetadata::FVideo Video;
+	TValueOrError<TUniquePtr<IVideoReader>, FText> VideoReaderResult = MediaManager.CreateVideoReader(InCurrentTakeFile);
+	if (VideoReaderResult.HasValue())
+	{
+		TUniquePtr<IVideoReader> VideoReader = VideoReaderResult.StealValue();
+		Video.FrameRate = static_cast<float>(VideoReader->GetFrameRate().AsDecimal());
+		TValueOrError<TUniquePtr<UE::CaptureManager::FMediaTextureSample>, FText> TextureSampleResult = VideoReader->Next();
+		if (TextureSampleResult.HasValue())
+		{
+			FVideoDeviceThumbnailExtractor ThumbnailExtractor;
+			UE::CaptureManager::FMediaTextureSample* Sample = TextureSampleResult.GetValue().Get();
+
+			TArray<FColor> ThumbnailRawColorData = ThumbnailExtractor.ExtractThumbnail(Sample);
+
+			if (ThumbnailRawColorData.Num() > 0)
+			{
+				FTakeThumbnailData::FRawImage RawImage;
+				RawImage.DecompressedImageData = MoveTemp(ThumbnailRawColorData);
+				RawImage.Width = Sample->Dimensions.X;
+				RawImage.Height = Sample->Dimensions.Y;
+				RawImage.Format = ERawImageFormat::BGRA8;
+
+				TakeMetadata.Thumbnail = MoveTemp(RawImage);
+			}
+		}
+	}
+	else
+	{
+		UE_LOG(LogVideoLiveLinkDevice, Warning, TEXT("Failed to determine the frame rate for the video file %s. Consider enabling Third Party Encoder in Capture Manager settings."), *InCurrentTakeFile);
+	}
+
+	Video.Name = MoveTemp(Name);
+	Video.Format = FPaths::GetExtension(InCurrentTakeFile);
+	Video.Path = InCurrentTakeFile;
+	Video.PathType = FTakeMetadata::FVideo::EPathType::File;
+
+	TSharedPtr<FCaptureExtractVideoTimecode> Extractor = MakeShareable(new FCaptureExtractVideoTimecode(Video.Path));
+	FTimecodeAndRateResult TimecodeAndRateResult = Extractor->Extract();
+	if (TimecodeAndRateResult.IsValid())
+	{
+		FTimecodeAndRate TimecodeAndRate = TimecodeAndRateResult.GetValue();
+		Video.TimecodeStart = TimecodeAndRate.Timecode.ToString();
+	}
+
+	TakeMetadata.Video.Add(MoveTemp(Video));
+
+	// Check if the video file contains audio stream by creating a audio reader from the file
+	// If creation fails, we assume that there is no audio in the file
+	TValueOrError<TUniquePtr<IAudioReader>, FText> AudioReaderResult = MediaManager.CreateAudioReader(InCurrentTakeFile);
+	
+	if (AudioReaderResult.HasValue())
+	{
+		FTakeMetadata::FAudio Audio;
+		Audio.Name = TEXT("audio");
+		Audio.Path = InCurrentTakeFile;
+		Audio.Duration = static_cast<float>(AudioReaderResult.StealValue()->GetDuration().GetTotalSeconds());
+
+		TakeMetadata.Audio.Add(MoveTemp(Audio));
+	}
+
+	return TakeMetadata;
+}
+
+#undef LOCTEXT_NAMESPACE
diff --git a/Engine/Plugins/VirtualProduction/CaptureManager/CaptureManagerDevices/Source/MonoVideoIngestDevice/Private/MonoVideoIngestDevice.h b/Engine/Plugins/VirtualProduction/CaptureManager/CaptureManagerDevices/Source/MonoVideoIngestDevice/Private/MonoVideoIngestDevice.h
new file mode 100644
--- /dev/null
+++ b/Engine/Plugins/VirtualProduction/CaptureManager/CaptureManagerDevices/Source/MonoVideoIngestDevice/Private/MonoVideoIngestDevice.h
@@ -0,0 +1,105 @@
+// Copyright Epic Games, Inc. All Rights Reserved.
+
+#pragma once
+
+#include "LiveLinkDevice.h"
+
+#include "Ingest/LiveLinkDeviceCapability_Ingest.h"
+#include "LiveLinkDeviceCapability_Connection.h"
+#include "Customizations/TakeDiscoveryExpressionCustomization.h"
+
+#include "Misc/Optional.h"
+
+#include "Async/TaskProgress.h"
+
+#include "BaseIngestLiveLinkDevice.h"
+
+#include "MonoVideoIngestDevice.generated.h"
+
+
+UCLASS(BlueprintType)
+class UMonoVideoIngestDeviceSettings : public ULiveLinkDeviceSettings
+{
+	GENERATED_BODY()
+
+public:
+	UMonoVideoIngestDeviceSettings()
+	{
+		VideoDiscoveryExpression = TEXT("<Auto>");
+	}
+
+	UPROPERTY(EditAnywhere, BlueprintReadWrite, Category = "Mono Video Ingest")
+	FString DisplayName = TEXT("Mono Video Ingest");
+
+	/* Path to a directory containing the take(s) data */
+	UPROPERTY(EditAnywhere, BlueprintReadWrite, Category = "Mono Video Ingest")
+	FDirectoryPath TakeDirectory;
+
+	/**
+	 * Format of the file name that the device expects
+	 *		Available expressions:
+	 *			* <Slate> - slate name
+	 *			* <Name> - identifier for the camera
+	 *			* <Take> - take number
+	 *			* <Any> - any string
+	 *			* <Auto> - used alone to automatically figure out takes based on directory structure
+	 *		Available delimiters: '_', '-', '.'
+	 *
+	 * Examples:
+	 * Considering input video MyTakeFolder/MySlate_MyName_SomeString-005.mov
+	 *
+	 * Video File Name Pattern <Auto> - will set Slate to the video filename “MySlate_MyName_SomeString_005” (dash converted to underline), Name will be set to just "video" and Take will be set to 1 by default.
+	 *
+	 * Video File Name Pattern <Slate>_<Name>_<Any>-<Take> - will extract Slate as “MySlate”, Name as “MyName” and Take as 5. <Any> ensures that the “SomeString” part of the filename is ignored.
+	 *
+	 * In both cases, extension '.mov' is stripped out.
+	 *
+	 * Note: when not using <Auto> expression, then both <Slate> and <Name> must be present in the pattern.
+	 */
+	UPROPERTY(EditAnywhere, BlueprintReadWrite, Category = "Mono Video Ingest")
+	FTakeDiscoveryExpression VideoDiscoveryExpression;
+};
+
+
+UCLASS(BlueprintType, meta = (DisplayName = "Mono Video Ingest", ToolTip = "Ingest video files as mono takes"))
+class UMonoVideoIngestDevice final : public UBaseIngestLiveLinkDevice
+	, public ILiveLinkDeviceCapability_Connection
+{
+	GENERATED_BODY()
+
+public:
+
+	UFUNCTION(BlueprintCallable, Category = "GenericArchiveDeviceSettings")
+	const UMonoVideoIngestDeviceSettings* GetSettings() const;
+
+	//~ Begin ULiveLinkDevice interface
+	virtual TSubclassOf<ULiveLinkDeviceSettings> GetSettingsClass() const override;
+	virtual FText GetDisplayName() const override;
+	virtual EDeviceHealth GetDeviceHealth() const override;
+	virtual FText GetHealthText() const override;
+
+	virtual void OnSettingChanged(const FPropertyChangedEvent& InPropertyChangedEvent) override;
+	//~ End ULiveLinkDevice interface
+
+private:
+	virtual FString GetFullTakePath(UE::CaptureManager::FTakeId InTakeId) const override;
+
+	//~ Begin ULiveLinkDeviceCapability_Ingest interface
+	virtual void RunConvertAndUploadTake(const UIngestCapability_ProcessHandle* InProcessHandle, const UIngestCapability_Options* InIngestOptions) override;
+	virtual void UpdateTakeList_Implementation(UIngestCapability_UpdateTakeListCallback* InCallback) override;
+	//~ End ULiveLinkDeviceCapability_Ingest interface
+
+	//~ Begin ILiveLinkDeviceCapability_Connection interface
+	virtual ELiveLinkDeviceConnectionStatus GetConnectionStatus_Implementation() const override;
+	virtual FString GetHardwareId_Implementation() const override;
+	virtual bool SetHardwareId_Implementation(const FString& HardwareID) override;
+	virtual bool Connect_Implementation() override;
+	virtual bool Disconnect_Implementation() override;
+	//~ End ILiveLinkDeviceCapability_Connection interface
+
+	bool IsVideoFile(const FString& InFileNameWithExtension);
+
+	TOptional<FTakeMetadata> ReadTake(const FString& InCurrentVideoFile) const;
+
+	TMap<int32, FString> FullTakePaths;
+};
diff --git a/Engine/Plugins/VirtualProduction/CaptureManager/CaptureManagerDevices/Source/MonoVideoIngestDevice/Private/MonoVideoIngestDeviceModule.cpp b/Engine/Plugins/VirtualProduction/CaptureManager/CaptureManagerDevices/Source/MonoVideoIngestDevice/Private/MonoVideoIngestDeviceModule.cpp
new file mode 100644
--- /dev/null
+++ b/Engine/Plugins/VirtualProduction/CaptureManager/CaptureManagerDevices/Source/MonoVideoIngestDevice/Private/MonoVideoIngestDeviceModule.cpp
@@ -0,0 +1,34 @@
+// Copyright Epic Games, Inc. All Rights Reserved.
+
+#include "Modules/ModuleManager.h"
+
+#include "PropertyEditorModule.h"
+#include "Customizations/TakeDiscoveryExpressionCustomization.h"
+#include "Templates/SharedPointer.h"
+
+class FMonoVideoIngestDeviceModule : public IModuleInterface
+{
+public:
+
+	void StartupModule()
+	{
+		FPropertyEditorModule& PropertyModule = FModuleManager::LoadModuleChecked<FPropertyEditorModule>("PropertyEditor");
+		PropertyModule.RegisterCustomPropertyTypeLayout(FTakeDiscoveryExpression::StaticStruct()->GetFName(),
+			FOnGetPropertyTypeCustomizationInstance::CreateRaw(this, &FMonoVideoIngestDeviceModule::GetTakeDiscoveryExpressionCustomization));
+	}
+
+	void ShutdownModule()
+	{
+		FPropertyEditorModule& PropertyModule = FModuleManager::LoadModuleChecked<FPropertyEditorModule>("PropertyEditor");
+		PropertyModule.UnregisterCustomPropertyTypeLayout(FTakeDiscoveryExpression::StaticStruct()->GetFName());
+	}
+
+private:
+
+	TSharedRef<IPropertyTypeCustomization> GetTakeDiscoveryExpressionCustomization()
+	{
+		return MakeShared<FTakeDiscoveryExpressionCustomization>();
+	}
+};
+
+IMPLEMENT_MODULE(FMonoVideoIngestDeviceModule, MonoVideoIngestDevice)
\ No newline at end of file
diff --git a/Engine/Plugins/VirtualProduction/CaptureManager/CaptureManagerDevices/Source/MonoVideoLiveLinkDevice/MonoVideoLiveLinkDevice.Build.cs b/Engine/Plugins/VirtualProduction/CaptureManager/CaptureManagerDevices/Source/MonoVideoLiveLinkDevice/MonoVideoLiveLinkDevice.Build.cs
deleted file mode 100644
diff --git a/Engine/Plugins/VirtualProduction/CaptureManager/CaptureManagerDevices/Source/MonoVideoLiveLinkDevice/Private/MonoVideoLiveLinkDevice.cpp b/Engine/Plugins/VirtualProduction/CaptureManager/CaptureManagerDevices/Source/MonoVideoLiveLinkDevice/Private/MonoVideoLiveLinkDevice.cpp
deleted file mode 100644
diff --git a/Engine/Plugins/VirtualProduction/CaptureManager/CaptureManagerDevices/Source/MonoVideoLiveLinkDevice/Private/MonoVideoLiveLinkDevice.h b/Engine/Plugins/VirtualProduction/CaptureManager/CaptureManagerDevices/Source/MonoVideoLiveLinkDevice/Private/MonoVideoLiveLinkDevice.h
deleted file mode 100644
diff --git a/Engine/Plugins/VirtualProduction/CaptureManager/CaptureManagerDevices/Source/MonoVideoLiveLinkDevice/Private/MonoVideoLiveLinkDeviceModule.cpp b/Engine/Plugins/VirtualProduction/CaptureManager/CaptureManagerDevices/Source/MonoVideoLiveLinkDevice/Private/MonoVideoLiveLinkDeviceModule.cpp
deleted file mode 100644
diff --git a/Engine/Plugins/VirtualProduction/CaptureManager/CaptureManagerDevices/Source/StereoVideoIngestDevice/Private/StereoVideoIngestDevice.cpp b/Engine/Plugins/VirtualProduction/CaptureManager/CaptureManagerDevices/Source/StereoVideoIngestDevice/Private/StereoVideoIngestDevice.cpp
new file mode 100644
--- /dev/null
+++ b/Engine/Plugins/VirtualProduction/CaptureManager/CaptureManagerDevices/Source/StereoVideoIngestDevice/Private/StereoVideoIngestDevice.cpp
@@ -0,0 +1,628 @@
+// Copyright Epic Games, Inc. All Rights Reserved.
+
+#include "StereoVideoIngestDevice.h"
+
+#include "HAL/FileManager.h"
+#include "Settings/CaptureManagerSettings.h"
+#include "ImageUtils.h"
+#include "CaptureManagerMediaRWModule.h"
+#include "Asset/CaptureAssetSanitization.h"
+#include "Misc/FileHelper.h"
+#include "Engine/Texture2D.h"
+
+#include "Misc/MessageDialog.h"
+#include "Utils/TakeDiscoveryExpressionParser.h"
+#include "Utils/VideoDeviceThumbnailExtractor.h"
+#include "Utils/CaptureExtractTimecode.h"
+#include "Utils/ParseTakeUtils.h"
+
+#include "VideoLiveLinkDeviceLog.h"
+
+DEFINE_LOG_CATEGORY(LogVideoLiveLinkDevice);
+#define LOCTEXT_NAMESPACE "StereoVideoLLDevice"
+
+namespace UE::StereoVideoLiveLinkDevice::Private
+{
+	static const TArray<FString::ElementType> Delimiters =
+	{
+		'-',
+		'_',
+		'.',
+		'/'
+	};
+
+	static const TArray<FStringView> SupportedVideoExtensions =
+	{
+		TEXTVIEW("mp4"),
+		TEXTVIEW("mov")
+	};
+
+	static const TArray<FStringView> SupportedAudioExtensions =
+	{
+		TEXTVIEW("wav")
+	};
+
+	static const TArray<FStringView> SupportedImageSequenceExtensions =
+	{
+		TEXTVIEW("jpg"),
+		TEXTVIEW("jpeg"),
+		TEXTVIEW("png")
+	};
+
+	FMediaRWManager& GetMediaRWManager()
+	{
+		return FModuleManager::LoadModuleChecked<FCaptureManagerMediaRWModule>("CaptureManagerMediaRW").Get();
+	}
+}
+
+const UStereoVideoIngestDeviceSettings* UStereoVideoIngestDevice::GetSettings() const
+{
+	return GetDeviceSettings<UStereoVideoIngestDeviceSettings>();
+}
+
+TSubclassOf<ULiveLinkDeviceSettings> UStereoVideoIngestDevice::GetSettingsClass() const
+{
+	return UStereoVideoIngestDeviceSettings::StaticClass();
+}
+
+FText UStereoVideoIngestDevice::GetDisplayName() const
+{
+	return FText::FromString(GetSettings()->DisplayName);
+}
+
+EDeviceHealth UStereoVideoIngestDevice::GetDeviceHealth() const
+{
+	return EDeviceHealth::Nominal;
+}
+
+FText UStereoVideoIngestDevice::GetHealthText() const
+{
+	return FText::FromString("Nominal");
+}
+
+void UStereoVideoIngestDevice::OnSettingChanged(const FPropertyChangedEvent& InPropertyChangedEvent)
+{
+	Super::OnSettingChanged(InPropertyChangedEvent);
+
+	const FName& PropertyName = InPropertyChangedEvent.GetMemberPropertyName();
+
+	if (GET_MEMBER_NAME_CHECKED(UStereoVideoIngestDeviceSettings, TakeDirectory) == PropertyName
+		|| GET_MEMBER_NAME_CHECKED(UStereoVideoIngestDeviceSettings, VideoDiscoveryExpression) == PropertyName
+		|| GET_MEMBER_NAME_CHECKED(UStereoVideoIngestDeviceSettings, AudioDiscoveryExpression) == PropertyName)
+	{
+		SetConnectionStatus(ELiveLinkDeviceConnectionStatus::Disconnected);
+
+		FString Path = GetSettings()->TakeDirectory.Path;
+		if (!Path.IsEmpty() && FPaths::DirectoryExists(Path))
+		{
+			SetConnectionStatus(ELiveLinkDeviceConnectionStatus::Connected);
+		}
+		else
+		{
+			SetConnectionStatus(ELiveLinkDeviceConnectionStatus::Disconnected);
+		}
+	}
+}
+
+FString UStereoVideoIngestDevice::GetFullTakePath(UE::CaptureManager::FTakeId InTakeId) const
+{
+	if (const FString* TakePath = FullTakePaths.Find(InTakeId))
+	{
+		return *TakePath;
+	}
+
+	return FString();
+}
+
+void UStereoVideoIngestDevice::RunConvertAndUploadTake(const UIngestCapability_ProcessHandle* InProcessHandle, const UIngestCapability_Options* InIngestOptions)
+{
+	AsyncTask(ENamedThreads::AnyThread, [this,
+		ProcessHandle = TStrongObjectPtr<const UIngestCapability_ProcessHandle>(InProcessHandle),
+		IngestOptions = TStrongObjectPtr<const UIngestCapability_Options>(InIngestOptions)]()
+		{
+			static constexpr uint32 NumberOfTasks = 2;
+
+			using namespace UE::CaptureManager;
+			TSharedPtr<FTaskProgress> TaskProgress
+				= MakeShared<FTaskProgress>(NumberOfTasks, FTaskProgress::FProgressReporter::CreateLambda([this, ProcessHandle](double InProgress)
+					{
+						ExecuteProcessProgressReporter(ProcessHandle.Get(), InProgress);
+					}));
+
+			Super::IngestTake(ProcessHandle.Get(), IngestOptions.Get(), MoveTemp(TaskProgress));
+		});
+}
+
+int32 UStereoVideoIngestDevice::FTakeWithComponents::CountComponents(UStereoVideoIngestDevice::ETakeComponentType Type)
+{
+	int32 Count = 0;
+	for (const UStereoVideoIngestDevice::FTakeWithComponents::Component& Component : Components)
+	{
+		if (Component.Type == Type)
+		{
+			++Count;
+		}
+	}
+	return Count;
+}
+
+void UStereoVideoIngestDevice::UpdateTakeList_Implementation(UIngestCapability_UpdateTakeListCallback* InCallback)
+{
+	AsyncTask(ENamedThreads::AnyThread, [this, Callback = TStrongObjectPtr<UIngestCapability_UpdateTakeListCallback>(InCallback)]()
+	{
+		RemoveAllTakes();
+
+		FString StoragePath = GetSettings()->TakeDirectory.Path;
+
+		TMap<FString, UStereoVideoIngestDevice::FTakeWithComponents> TakesWithComponentsCandidates = DiscoverTakes(StoragePath);
+
+		TArray<FString> SlateNames;
+		TakesWithComponentsCandidates.GetKeys(SlateNames);
+
+		for (const FString& SlateName : SlateNames)
+		{
+			FTakeWithComponents& TakeWithComponents = TakesWithComponentsCandidates[SlateName];
+
+			bool TakeAsExpected = ((TakeWithComponents.CountComponents(ETakeComponentType::VIDEO) == 2 && TakeWithComponents.CountComponents(ETakeComponentType::IMAGE_SEQUENCE) == 0) ||
+									(TakeWithComponents.CountComponents(ETakeComponentType::VIDEO) == 0 && TakeWithComponents.CountComponents(ETakeComponentType::IMAGE_SEQUENCE) == 2)) &&
+				TakeWithComponents.CountComponents(ETakeComponentType::AUDIO) <= 1;
+
+			if (!TakeAsExpected)
+			{
+				UE_LOG(LogVideoLiveLinkDevice, Warning, TEXT("Invalid take at '%s'. Take should have exactly two video files or two image sequences. Optionally, one audio file."), *TakeWithComponents.TakePath);
+				continue;
+			}
+
+			TOptional<FTakeMetadata> TakeInfoResult = CreateTakeMetadata(TakeWithComponents);
+
+			if (TakeInfoResult.IsSet())
+			{
+				FTakeMetadata TakeMetadata = TakeInfoResult.GetValue();
+
+				int32 TakeId = AddTake(MoveTemp(TakeMetadata));
+				FullTakePaths.Add(TakeId, TakeWithComponents.TakePath);
+
+				PublishEvent<FIngestCapability_TakeAddedEvent>(TakeId);
+			}
+		}
+
+		ExecuteUpdateTakeListCallback(Callback.Get(), Execute_GetTakeIdentifiers(this));
+	});
+}
+
+TMap<FString, UStereoVideoIngestDevice::FTakeWithComponents> UStereoVideoIngestDevice::DiscoverTakes(FString StoragePath)
+{
+	TMap<FString, FTakeWithComponents> TakesWithComponentsCandidates;
+
+	int32 DirectoriesInterrogatedCount = 0;
+	const int32 DirectoriesToInterrogateInOneRun = 200;
+
+	const bool bIterationResult = IFileManager::Get().IterateDirectoryRecursively(*StoragePath, [&](const TCHAR* InPath, bool bInIsDirectory)
+		{
+			if (bInIsDirectory)
+			{
+				if (++DirectoriesInterrogatedCount > DirectoriesToInterrogateInOneRun)
+				{
+					return false;
+				}
+			}
+
+			FString Extension = FPaths::GetExtension(InPath);
+			if (UE::StereoVideoLiveLinkDevice::Private::SupportedVideoExtensions.Contains(Extension))
+			{
+				ExtractTakeComponents(InPath, StoragePath, ETakeComponentType::VIDEO, GetSettings()->VideoDiscoveryExpression.Value, "UnknownVideoName", TakesWithComponentsCandidates);
+			}
+			else if (UE::StereoVideoLiveLinkDevice::Private::SupportedImageSequenceExtensions.Contains(Extension))
+			{
+				ExtractTakeComponents(InPath, StoragePath, ETakeComponentType::IMAGE_SEQUENCE, GetSettings()->VideoDiscoveryExpression.Value, "UnknownImageSequenceName", TakesWithComponentsCandidates);
+			}
+			else if (UE::StereoVideoLiveLinkDevice::Private::SupportedAudioExtensions.Contains(Extension))
+			{
+				ExtractTakeComponents(InPath, StoragePath, ETakeComponentType::AUDIO, GetSettings()->AudioDiscoveryExpression.Value, "UnknownAudioName", TakesWithComponentsCandidates);
+			}
+
+			return true;
+		});
+
+	return TakesWithComponentsCandidates;
+}
+
+void UStereoVideoIngestDevice::ExtractTakeComponents(FString ComponentPath, FString StoragePath, ETakeComponentType ComponentType, FString Format, FString UnknownComponentName, TMap<FString, UStereoVideoIngestDevice::FTakeWithComponents>& OutTakesWithComponentsCandidates)
+{
+	TOptional<FTakeWithComponents> TakeWithComponents;
+
+	if (Format == TEXT("<Auto>"))
+	{
+		TakeWithComponents = ExtractTakeComponentsFromDirectoryStructure(ComponentPath, StoragePath, ComponentType);
+	}
+	else
+	{
+		TakeWithComponents = ExtractTakeComponentsUsingTokens(ComponentPath, StoragePath, Format, ComponentType);
+	}
+	
+	if (!TakeWithComponents.IsSet())
+	{
+		FString ComponentRootPath;
+		FString FileName;
+		FString Extension;
+		FPaths::Split(ComponentPath, ComponentRootPath, FileName, Extension);
+
+		TakeWithComponents = { ComponentRootPath, "Slate name not determined", -1, {{UnknownComponentName, ComponentType, ComponentPath}} };
+	}
+
+	GroupFoundComponents(OutTakesWithComponentsCandidates, TakeWithComponents.GetValue());
+};
+
+TOptional<UStereoVideoIngestDevice::FTakeWithComponents> UStereoVideoIngestDevice::ExtractTakeComponentsUsingTokens(FString ComponentPath, FString StoragePath, FString Format, ETakeComponentType ComponentType)
+{
+	FStringView InPathStringView(ComponentPath);
+	FString RelativePath = FString(InPathStringView.SubStr(StoragePath.Len(), InPathStringView.Len() - StoragePath.Len()));
+	FString Extension = FPaths::GetExtension(RelativePath);
+	FString RelativePathNoExtension = RelativePath.Mid(0, RelativePath.Len() - Extension.Len() - 1);
+	
+	FTakeDiscoveryExpressionParser TokenParser(Format, RelativePathNoExtension, UE::StereoVideoLiveLinkDevice::Private::Delimiters);
+
+	FString StorageLeafWithRelativePath = FPaths::GetPathLeaf(StoragePath) + RelativePathNoExtension;
+	FTakeDiscoveryExpressionParser WithLeafTokenParser(Format, StorageLeafWithRelativePath, UE::StereoVideoLiveLinkDevice::Private::Delimiters);
+
+	FString TakePath;
+	FTakeDiscoveryExpressionParser* SuccessfulParser = nullptr;
+	if (TokenParser.Parse())
+	{
+		SuccessfulParser = &TokenParser;
+		TArray<FString> RelativePathParts;
+		FPaths::NormalizeDirectoryName(RelativePath);
+		RelativePath.ParseIntoArray(RelativePathParts, TEXT("/"));
+		TakePath = StoragePath + TEXT("/") + RelativePathParts[0];
+	}
+	else if (WithLeafTokenParser.Parse())
+	{
+		SuccessfulParser = &WithLeafTokenParser;
+		TakePath = StoragePath;
+	}
+
+	if (SuccessfulParser != nullptr)
+	{
+		return FTakeWithComponents{ TakePath, SuccessfulParser->GetSlateName(), SuccessfulParser->GetTakeNumber(), { { SuccessfulParser->GetName(), ComponentType, ComponentPath } } };
+	}
+
+	return {};
+};
+
+TOptional<UStereoVideoIngestDevice::FTakeWithComponents> UStereoVideoIngestDevice::ExtractTakeComponentsFromDirectoryStructure(FString ComponentPath, FString StoragePath, ETakeComponentType ComponentType)
+{
+	FStringView InPathStringView(ComponentPath);
+	FString RelativePath = FString(InPathStringView.SubStr(StoragePath.Len(), InPathStringView.Len() - StoragePath.Len()));
+
+	FString TakePath;
+	FString TakeName;
+
+	if (ComponentType == ETakeComponentType::VIDEO || ComponentType == ETakeComponentType::AUDIO)
+	{
+		// Use video file container folder name as TakeName
+		// this is for case where user navigate StoragePath to a container folder
+		TArray<FString> Parts;
+		FPaths::NormalizeDirectoryName(RelativePath);
+		RelativePath.ParseIntoArray(Parts, TEXT("/"));
+		if (Parts.Num() > 1)
+		{
+			TakeName = Parts[0];
+			TakePath = StoragePath + "/" + TakeName;
+		}
+		else
+		{
+			// If the above fails, because user navigates StoragePath to TakeFolder, 
+			// we could use specified storage path leaf folder as TakeName
+			TakeName = FPaths::GetPathLeaf(StoragePath);
+			TakePath = StoragePath;
+		}
+	}
+	else if (ComponentType == ETakeComponentType::IMAGE_SEQUENCE)
+	{
+		// Assumption is that frames are stored in separate folders
+		TArray<FString> Parts;
+		FPaths::NormalizeDirectoryName(RelativePath);
+		RelativePath.ParseIntoArray(Parts, TEXT("/"));
+		if (Parts.Num() > 2)
+		{
+			TakeName = Parts[0];
+			TakePath = StoragePath + "/" + TakeName;
+		}
+		else
+		{
+			TakePath = StoragePath;
+			TakeName = FPaths::GetPathLeaf(StoragePath);
+		}
+
+		{
+			FString Filename, Extension;
+			FPaths::Split(ComponentPath, ComponentPath, Filename, Extension);
+		}
+	}
+
+	if (TakeName.IsEmpty())
+	{
+		return {}; // Failed to match TakeName
+	}
+
+	FString SlateName = TakeName;
+	int32 TakeNumber = -1;
+	FString Name;
+
+	return FTakeWithComponents{ TakePath, SlateName, TakeNumber, { { Name, ComponentType, ComponentPath } } };
+};
+
+void UStereoVideoIngestDevice::GroupFoundComponents(TMap<FString, UStereoVideoIngestDevice::FTakeWithComponents>& TakesWithComponentsCandidates, FTakeWithComponents TakeWithComponents)
+{
+	FString TakeIdentifier = TakeWithComponents.SlateName + FString::FromInt(TakeWithComponents.TakeNumber);
+
+	FTakeWithComponents* TakeWithComponentsPtr = TakesWithComponentsCandidates.Find(TakeIdentifier);
+	if (TakeWithComponentsPtr == nullptr)
+	{
+		TakeWithComponentsPtr = &TakesWithComponentsCandidates.Add(TakeIdentifier, { TakeWithComponents.TakePath, TakeWithComponents.SlateName, TakeWithComponents.TakeNumber, {} });
+	}
+
+	for (FTakeWithComponents::Component& Component : TakeWithComponents.Components)
+	{
+		FTakeWithComponents::Component* TakeComponentPtr = TakeWithComponentsPtr->Components.FindByPredicate(
+			[&Component](const FTakeWithComponents::Component& Other)
+			{
+				return Other.Path == Component.Path;
+			});
+
+		if (TakeComponentPtr == nullptr)
+		{
+			TakeWithComponentsPtr->Components.Add(MoveTemp(Component));
+		}
+	}
+};
+
+ELiveLinkDeviceConnectionStatus UStereoVideoIngestDevice::GetConnectionStatus_Implementation() const
+{
+	const UStereoVideoIngestDeviceSettings* DeviceSettings = GetSettings();
+	FString Path = DeviceSettings->TakeDirectory.Path;
+
+	if (!Path.IsEmpty() && FPaths::DirectoryExists(Path))
+	{
+		return ELiveLinkDeviceConnectionStatus::Connected;
+	}
+
+	return ELiveLinkDeviceConnectionStatus::Disconnected;
+}
+
+FString UStereoVideoIngestDevice::GetHardwareId_Implementation() const
+{
+	return FPlatformMisc::GetDeviceId();
+}
+
+bool UStereoVideoIngestDevice::SetHardwareId_Implementation(const FString& HardwareID)
+{
+	return false;
+}
+
+bool UStereoVideoIngestDevice::Connect_Implementation()
+{
+	const UStereoVideoIngestDeviceSettings* DeviceSettings = GetSettings();
+	FString Path = DeviceSettings->TakeDirectory.Path;
+
+	if (!Path.IsEmpty() && FPaths::DirectoryExists(Path))
+	{
+		SetConnectionStatus(ELiveLinkDeviceConnectionStatus::Connected);
+		return true;
+	}
+
+	return false;
+}
+
+bool UStereoVideoIngestDevice::Disconnect_Implementation()
+{
+	SetConnectionStatus(ELiveLinkDeviceConnectionStatus::Disconnected);
+	return true;
+}
+
+bool UStereoVideoIngestDevice::IsVideoFile(const FString& InFileNameWithExtension)
+{
+	FString Extension = FPaths::GetExtension(InFileNameWithExtension);
+
+	return UE::StereoVideoLiveLinkDevice::Private::SupportedVideoExtensions.Contains(Extension);
+}
+
+bool UStereoVideoIngestDevice::IsFrameInSequenceFile(const FString& InFileNameWithExtension)
+{
+	FString Extension = FPaths::GetExtension(InFileNameWithExtension);
+
+	return UE::StereoVideoLiveLinkDevice::Private::SupportedImageSequenceExtensions.Contains(Extension);
+}
+
+bool UStereoVideoIngestDevice::IsAudioFile(const FString& InFileNameWithExtension)
+{
+	FString Extension = FPaths::GetExtension(InFileNameWithExtension);
+
+	return UE::StereoVideoLiveLinkDevice::Private::SupportedAudioExtensions.Contains(Extension);
+}
+
+TOptional<FTakeMetadata> UStereoVideoIngestDevice::CreateTakeMetadata(const FTakeWithComponents& InTakeWithComponents) const
+{
+	using namespace UE::CaptureManager;
+
+	FMediaRWManager& MediaManager = UE::StereoVideoLiveLinkDevice::Private::GetMediaRWManager();
+
+	FTakeMetadata TakeMetadata;
+
+	int32 TakeNumber = InTakeWithComponents.TakeNumber;
+
+	if (TakeNumber == INDEX_NONE)
+	{
+		TakeNumber = 1;
+	}
+
+	TakeMetadata.Version.Major = 4;
+	TakeMetadata.Version.Minor = 1;
+
+	TakeMetadata.Slate = InTakeWithComponents.SlateName;
+	TakeMetadata.TakeNumber = TakeNumber;
+
+	// Create new guid
+	TakeMetadata.UniqueId = FGuid::NewGuid().ToString(EGuidFormats::DigitsWithHyphensLower);
+
+	FString InCurrentTakeFile = InTakeWithComponents.Components[0].Path;
+	FFileStatData FileData = IFileManager::Get().GetStatData(*InCurrentTakeFile);
+	TakeMetadata.DateTime = FileData.CreationTime;
+
+	TakeMetadata.Device.Model = TEXT("StereoHMC");
+
+	int32 AudioNameCounter = 1;
+
+	bool bVideoFrameRateSet = false;
+	FFrameRate VideoFrameRate;
+
+	for (const FTakeWithComponents::Component& TakeComponent : InTakeWithComponents.Components)
+	{
+		if (TakeComponent.Type == ETakeComponentType::VIDEO || TakeComponent.Type == ETakeComponentType::IMAGE_SEQUENCE)
+		{
+			FTakeMetadata::FVideo Video;
+
+			FString Name = TakeComponent.Name;
+			if (Name.IsEmpty())
+			{
+				Name = TakeComponent.Path.Mid(InTakeWithComponents.TakePath.Len() + 1);
+				SanitizePackagePath(Name, '_');
+			}
+
+			Video.Name = MoveTemp(Name);
+			Video.Format = FPaths::GetExtension(TakeComponent.Path);
+
+			if (TakeComponent.Type == ETakeComponentType::VIDEO)
+			{
+				Video.PathType = FTakeMetadata::FVideo::EPathType::File;
+
+				TValueOrError<TUniquePtr<IVideoReader>, FText> VideoReaderResult = MediaManager.CreateVideoReader(TakeComponent.Path);
+				if (VideoReaderResult.HasValue())
+				{
+					TUniquePtr<IVideoReader> VideoReader = VideoReaderResult.StealValue();
+					Video.FrameRate = static_cast<float>(VideoReader->GetFrameRate().AsDecimal());
+
+					if (!bVideoFrameRateSet)
+					{
+						VideoFrameRate = UE::CaptureManager::ParseFrameRate(Video.FrameRate);
+					}
+
+					TValueOrError<TUniquePtr<UE::CaptureManager::FMediaTextureSample>, FText> TextureSampleResult = VideoReader->Next();
+					if (TextureSampleResult.HasValue())
+					{
+						FVideoDeviceThumbnailExtractor ThumbnailExtractor;
+						UE::CaptureManager::FMediaTextureSample* Sample = TextureSampleResult.GetValue().Get();
+
+						TArray<FColor> ThumbnailRawColorData = ThumbnailExtractor.ExtractThumbnail(Sample);
+
+						if (ThumbnailRawColorData.Num() > 0) 
+						{
+							FTakeThumbnailData::FRawImage RawImage;
+							RawImage.DecompressedImageData = MoveTemp(ThumbnailRawColorData);
+							RawImage.Width = Sample->Dimensions.X;
+							RawImage.Height = Sample->Dimensions.Y;
+							RawImage.Format = ERawImageFormat::BGRA8;
+
+							TakeMetadata.Thumbnail = MoveTemp(RawImage);
+						}
+					}
+				}
+				else
+				{
+					UE_LOG(LogVideoLiveLinkDevice, Warning, TEXT("Failed to determine the frame rate for the video file %s. Consider enabling Third Party Encoder in Capture Manager settings."), *TakeComponent.Path);
+				}
+
+				Video.Path = TakeComponent.Path;
+				
+				FString MediaFilePath = FPaths::ConvertRelativePathToFull(InTakeWithComponents.TakePath, Video.Path);
+				TSharedPtr<FCaptureExtractVideoTimecode> Extractor = MakeShareable(new FCaptureExtractVideoTimecode(MediaFilePath));
+				FTimecodeAndRateResult TimecodeAndRateResult = Extractor->Extract();
+				if (TimecodeAndRateResult.IsValid())
+				{
+					FTimecodeAndRate TimecodeAndRate = TimecodeAndRateResult.GetValue();
+					Video.TimecodeStart = TimecodeAndRate.Timecode.ToString();
+				}
+			}
+			else
+			{
+				Video.PathType = FTakeMetadata::FVideo::EPathType::Folder;
+
+				int32 FramesCount = 0;
+				bool FirstFrame = true;
+
+				bool bIterationResult = IFileManager::Get().IterateDirectoryRecursively(*TakeComponent.Path, [&](const TCHAR* InPath, bool bInIsDirectory)
+					{
+						if (FirstFrame)
+						{
+							FirstFrame = false;
+
+							if (!TakeMetadata.Thumbnail.GetThumbnailData().IsSet())
+							{
+								TArray<uint8> ThumbnailData;
+								if (FFileHelper::LoadFileToArray(ThumbnailData, InPath))
+								{
+									TakeMetadata.Thumbnail = ThumbnailData;
+								}
+							}
+
+							FString Path, Filename, Extension;
+							FPaths::Split(InPath, Path, Filename, Extension);
+							Video.Format = Extension;
+
+							Video.Path = TakeComponent.Path;
+							{
+								FImage Image;
+								FImageUtils::LoadImage(InPath, Image);
+								if (Image.IsImageInfoValid())
+								{
+									Video.FrameWidth = Image.SizeX;
+									Video.FrameHeight = Image.SizeY;
+								}
+							}
+						}
+
+						FramesCount++;
+
+						return true;
+					});
+
+				Video.FramesCount = FramesCount;
+			}
+
+			TakeMetadata.Video.Add(MoveTemp(Video));
+		}
+		else if (TakeComponent.Type == ETakeComponentType::AUDIO)
+		{
+			FTakeMetadata::FAudio Audio;
+			FString Name = TakeComponent.Name;
+			if (Name.IsEmpty())
+			{
+				Name = TEXT("audio") + FString::FromInt(AudioNameCounter++);
+			}
+			Audio.Name = Name;
+			Audio.Path = TakeComponent.Path;
+			Audio.Duration = 0;
+
+			FString MediaFilePath = FPaths::ConvertRelativePathToFull(InTakeWithComponents.TakePath, Audio.Path);
+			TSharedPtr<FCaptureExtractAudioTimecode> Extractor = MakeShareable(new FCaptureExtractAudioTimecode(MediaFilePath));
+
+			// The video frame rate will be used to calculate the timecode rate if the timecode rate cannot be extracted from the audio file
+			FTimecodeAndRateResult TimecodeAndRateResult = Extractor->Extract(VideoFrameRate);
+			if (TimecodeAndRateResult.IsValid())
+			{
+				FTimecodeAndRate TimecodeAndRate = TimecodeAndRateResult.GetValue();
+				Audio.TimecodeStart = TimecodeAndRate.Timecode.ToString();
+			
+				if (TimecodeAndRate.MaybeTimecodeRate.IsSet())
+				{
+					Audio.TimecodeRate = static_cast<float>(TimecodeAndRate.MaybeTimecodeRate.GetValue().AsDecimal());
+				}
+			}
+
+			TakeMetadata.Audio.Add(MoveTemp(Audio));
+		}
+	}
+	return TakeMetadata;
+}
+
+#undef LOCTEXT_NAMESPACE
diff --git a/Engine/Plugins/VirtualProduction/CaptureManager/CaptureManagerDevices/Source/StereoVideoIngestDevice/Private/StereoVideoIngestDevice.h b/Engine/Plugins/VirtualProduction/CaptureManager/CaptureManagerDevices/Source/StereoVideoIngestDevice/Private/StereoVideoIngestDevice.h
new file mode 100644
--- /dev/null
+++ b/Engine/Plugins/VirtualProduction/CaptureManager/CaptureManagerDevices/Source/StereoVideoIngestDevice/Private/StereoVideoIngestDevice.h
@@ -0,0 +1,169 @@
+// Copyright Epic Games, Inc. All Rights Reserved.
+
+#pragma once
+
+#include "LiveLinkDevice.h"
+
+#include "Ingest/LiveLinkDeviceCapability_Ingest.h"
+#include "LiveLinkDeviceCapability_Connection.h"
+#include "Customizations/TakeDiscoveryExpressionCustomization.h"
+
+#include "Misc/Optional.h"
+
+#include "Async/TaskProgress.h"
+
+#include "BaseIngestLiveLinkDevice.h"
+
+#include "StereoVideoIngestDevice.generated.h"
+
+
+UCLASS(BlueprintType)
+class UStereoVideoIngestDeviceSettings : public ULiveLinkDeviceSettings
+{
+	GENERATED_BODY()
+
+public:
+	UStereoVideoIngestDeviceSettings()
+	{
+		VideoDiscoveryExpression = TEXT("<Auto>");
+		AudioDiscoveryExpression = TEXT("<Auto>");
+	}
+
+	UPROPERTY(EditAnywhere, BlueprintReadWrite, Category = "Stereo Video Ingest")
+	FString DisplayName = TEXT("Stereo Video Ingest");
+
+	/* Path to a directory containing the take(s) data */
+	UPROPERTY(EditAnywhere, BlueprintReadWrite, Category = "Stereo Video Ingest")
+	FDirectoryPath TakeDirectory;
+
+	/**
+	 * Format of the file name that the device expects
+	 *		Available expressions:
+	 *			* <Slate> - slate name
+	 *			* <Name> - identifier for the camera in the stereo pair
+	 *			* <Take> - take number
+	 *			* <Any> - any string
+	 *			* <Auto> - used alone to automatically figure out takes based on directory structure
+	 *		Available delimiters: '_', '-', '.', '\'
+	 *
+	 * Examples:
+	 * Considering input video MyTakeFolder/MySlate_MyName_SomeString-005.mov
+	 *
+	 * Video File Name Pattern <Auto> - will set Slate to the containing folder “MyTakeFolder”, Name will be set to the video filename “MySlate_MyName_SomeString-005” and Take will be set to 1 by default
+	 *
+	 * Video File Name Pattern <Slate>_<Name>_<Any>-<Take> - will extract Slate as “MySlate”, Name as “MyName” and Take as 5. <Any> ensures that the “SomeString” part of the filename is ignored.
+	 *
+	 * In both cases, extension '.mov' is stripped out.
+	 * 
+	 * Note: when not using <Auto> expression, then both <Slate> and <Name> must be present in the pattern
+	 */
+	UPROPERTY(EditAnywhere, BlueprintReadWrite, Category = "Stereo Video Ingest")
+	FTakeDiscoveryExpression VideoDiscoveryExpression;
+
+	/**
+	 * Format of the file name that the device expects
+	 *		Available expressions:
+	 *			* <Slate> - slate name
+	 *			* <Name> - identifier for the camera in the stereo pair
+	 *			* <Take> - take number
+	 *			* <Any> - any string
+	 *			* <Auto> - used alone to automatically figure out takes based on directory structure
+	 *		Available delimiters: '_', '-', '.', '\'
+	 *
+	 * Examples:
+	 * Considering input audio MyTakeFolder/MySlate_MyName_SomeString-005.wav
+	 * 
+	  * Audio File Name Pattern <Auto> - will set Slate to the containing folder “MyTakeFolder”, Name will be set to the audio filename “MySlate_MyName_SomeString-005” and Take will be set to 1 by default
+	 *
+	 * Audio File Name Pattern <Slate>_<Name>_<Any>-<Take> - will extract Slate as “MySlate”, Name as “MyName” and Take as 5. <Any> ensures that the “SomeString” part of the filename is ignored.
+	 *
+	 * In both cases, extension '.wav' is stripped out.
+	 * 
+	 * Note: when not using <Auto> expression, then both <Slate> and <Name> must be present in the pattern
+	 */
+	UPROPERTY(EditAnywhere, BlueprintReadWrite, Category = "Stereo Video Ingest")
+	FTakeDiscoveryExpression AudioDiscoveryExpression;
+};
+
+
+UCLASS(BlueprintType, meta = (DisplayName = "Stereo Video Ingest", ToolTip = "Ingest subfolders containing pairs of video files as stereo takes"))
+class UStereoVideoIngestDevice final : public UBaseIngestLiveLinkDevice
+	, public ILiveLinkDeviceCapability_Connection
+{
+	GENERATED_BODY()
+
+public:
+	UStereoVideoIngestDevice() = default;
+	~UStereoVideoIngestDevice() = default;
+
+	UFUNCTION(BlueprintCallable, Category = "GenericArchiveDeviceSettings")
+	const UStereoVideoIngestDeviceSettings* GetSettings() const;
+
+	//~ Begin ULiveLinkDevice interface
+	virtual TSubclassOf<ULiveLinkDeviceSettings> GetSettingsClass() const override;
+	virtual FText GetDisplayName() const override;
+	virtual EDeviceHealth GetDeviceHealth() const override;
+	virtual FText GetHealthText() const override;
+
+	virtual void OnSettingChanged(const FPropertyChangedEvent& InPropertyChangedEvent) override;
+	//~ End ULiveLinkDevice interface
+
+private:
+
+	enum class ETakeComponentType
+	{
+		VIDEO,
+		IMAGE_SEQUENCE,
+		AUDIO
+	};
+
+	struct FTakeWithComponents
+	{
+		struct Component
+		{
+			FString Name;
+			ETakeComponentType Type;
+			FString Path;
+		};
+
+		FString TakePath;
+
+		FString SlateName;
+		int32 TakeNumber;
+
+		TArray<Component> Components;
+
+		int32 CountComponents(ETakeComponentType Type);
+	};
+	
+	using FTakesWithComponents = TArray<struct FTakeWithComponents>;
+
+	virtual FString GetFullTakePath(UE::CaptureManager::FTakeId InTakeId) const override;
+
+	//~ Begin ULiveLinkDeviceCapability_Ingest interface
+	virtual void RunConvertAndUploadTake(const UIngestCapability_ProcessHandle* InProcessHandle,
+										 const UIngestCapability_Options* InIngestOptions) override;
+	virtual void UpdateTakeList_Implementation(UIngestCapability_UpdateTakeListCallback* InCallback) override;
+	//~ End ULiveLinkDeviceCapability_Ingest interface
+
+	//~ Begin ILiveLinkDeviceCapability_Connection interface
+	virtual ELiveLinkDeviceConnectionStatus GetConnectionStatus_Implementation() const override;
+	virtual FString GetHardwareId_Implementation() const override;
+	virtual bool SetHardwareId_Implementation(const FString& HardwareID) override;
+	virtual bool Connect_Implementation() override;
+	virtual bool Disconnect_Implementation() override;
+	//~ End ILiveLinkDeviceCapability_Connection interface
+
+	bool IsVideoFile(const FString& InFileNameWithExtension);
+	bool IsFrameInSequenceFile(const FString& InFileNameWithExtension);
+	bool IsAudioFile(const FString& InFileNameWithExtension);
+
+	TOptional<FTakeWithComponents> ExtractTakeComponentsUsingTokens(FString ComponentPath, FString StoragePath, FString Format, ETakeComponentType ComponentType);
+	TOptional<FTakeWithComponents> ExtractTakeComponentsFromDirectoryStructure(FString ComponentPath, FString StoragePath, ETakeComponentType ComponentType);
+	void GroupFoundComponents(TMap<FString, UStereoVideoIngestDevice::FTakeWithComponents>& TakesWithComponentsCandidates, FTakeWithComponents TakeWithComponents);
+	void ExtractTakeComponents(FString ComponentPath, FString StoragePath, ETakeComponentType ComponentType, FString Format, FString UnknownComponentName, TMap<FString, UStereoVideoIngestDevice::FTakeWithComponents>& OutTakesWithComponentsCandidates);
+	TMap<FString, FTakeWithComponents> DiscoverTakes(FString StoragePath);
+	TOptional<FTakeMetadata> CreateTakeMetadata(const FTakeWithComponents & InTakesWithComponents) const;
+
+	TMap<int32, FString> FullTakePaths;
+};
diff --git a/Engine/Plugins/VirtualProduction/CaptureManager/CaptureManagerDevices/Source/StereoVideoIngestDevice/Private/StereoVideoIngestDeviceModule.cpp b/Engine/Plugins/VirtualProduction/CaptureManager/CaptureManagerDevices/Source/StereoVideoIngestDevice/Private/StereoVideoIngestDeviceModule.cpp
new file mode 100644
--- /dev/null
+++ b/Engine/Plugins/VirtualProduction/CaptureManager/CaptureManagerDevices/Source/StereoVideoIngestDevice/Private/StereoVideoIngestDeviceModule.cpp
@@ -0,0 +1,34 @@
+// Copyright Epic Games, Inc. All Rights Reserved.
+
+#include "Modules/ModuleManager.h"
+
+#include "PropertyEditorModule.h"
+#include "Customizations/TakeDiscoveryExpressionCustomization.h"
+#include "Templates/SharedPointer.h"
+
+class FStereoVideoIngestDeviceModule : public IModuleInterface
+{
+public:
+
+	void StartupModule()
+	{
+		FPropertyEditorModule& PropertyModule = FModuleManager::LoadModuleChecked<FPropertyEditorModule>("PropertyEditor");
+		PropertyModule.RegisterCustomPropertyTypeLayout(FTakeDiscoveryExpression::StaticStruct()->GetFName(),
+			FOnGetPropertyTypeCustomizationInstance::CreateRaw(this, &FStereoVideoIngestDeviceModule::GetTakeDiscoveryExpressionCustomization));
+	}
+
+	void ShutdownModule()
+	{
+		FPropertyEditorModule& PropertyModule = FModuleManager::LoadModuleChecked<FPropertyEditorModule>("PropertyEditor");
+		PropertyModule.UnregisterCustomPropertyTypeLayout(FTakeDiscoveryExpression::StaticStruct()->GetFName());
+	}
+
+private:
+
+	TSharedRef<IPropertyTypeCustomization> GetTakeDiscoveryExpressionCustomization()
+	{		
+		return MakeShared<FTakeDiscoveryExpressionCustomization>();
+	}
+};
+
+IMPLEMENT_MODULE(FStereoVideoIngestDeviceModule, StereoVideoIngestDevice)
\ No newline at end of file
diff --git a/Engine/Plugins/VirtualProduction/CaptureManager/CaptureManagerDevices/Source/StereoVideoIngestDevice/StereoVideoIngestDevice.Build.cs b/Engine/Plugins/VirtualProduction/CaptureManager/CaptureManagerDevices/Source/StereoVideoIngestDevice/StereoVideoIngestDevice.Build.cs
new file mode 100644
--- /dev/null
+++ b/Engine/Plugins/VirtualProduction/CaptureManager/CaptureManagerDevices/Source/StereoVideoIngestDevice/StereoVideoIngestDevice.Build.cs
@@ -0,0 +1,43 @@
+// Copyright Epic Games, Inc. All Rights Reserved.
+
+using UnrealBuildTool;
+using System.IO;
+
+public class StereoVideoIngestDevice : ModuleRules
+{
+	public StereoVideoIngestDevice(ReadOnlyTargetRules Target) : base(Target)
+	{
+		PCHUsage = PCHUsageMode.UseExplicitOrSharedPCHs;
+		ShortName = "StereoVideoIngest";
+
+		PublicDependencyModuleNames.AddRange(new string[]
+		{
+			"CaptureUtils",
+			"Core"
+		});
+
+		PrivateDependencyModuleNames.AddRange(new string[]
+		{
+			"Engine",
+			"Json",
+			"CaptureManagerMediaRW",
+			"IngestLiveLinkDevice",
+			"DataIngestCore",
+			"LiveLinkDevice",
+			"LiveLinkCapabilities",
+			"Media",
+			"LiveLinkHubCaptureMessaging",
+			"LiveLinkHub",
+			"ToolWidgets",
+			"Slate",
+			"SlateCore",
+			"CaptureManagerTakeMetadata",
+			"CaptureManagerSettings",
+			"CoreUObject",
+			"UnrealEd",
+			"VideoLiveLinkDeviceCommon"
+		});
+
+		CppCompileWarningSettings.UnsafeTypeCastWarningLevel = WarningLevel.Error;
+	}
+}
diff --git a/Engine/Plugins/VirtualProduction/CaptureManager/CaptureManagerDevices/Source/StereoVideoLiveLinkDevice/Private/StereoVideoLiveLinkDevice.cpp b/Engine/Plugins/VirtualProduction/CaptureManager/CaptureManagerDevices/Source/StereoVideoLiveLinkDevice/Private/StereoVideoLiveLinkDevice.cpp
deleted file mode 100644
diff --git a/Engine/Plugins/VirtualProduction/CaptureManager/CaptureManagerDevices/Source/StereoVideoLiveLinkDevice/Private/StereoVideoLiveLinkDevice.h b/Engine/Plugins/VirtualProduction/CaptureManager/CaptureManagerDevices/Source/StereoVideoLiveLinkDevice/Private/StereoVideoLiveLinkDevice.h
deleted file mode 100644
diff --git a/Engine/Plugins/VirtualProduction/CaptureManager/CaptureManagerDevices/Source/StereoVideoLiveLinkDevice/Private/StereoVideoLiveLinkDeviceModule.cpp b/Engine/Plugins/VirtualProduction/CaptureManager/CaptureManagerDevices/Source/StereoVideoLiveLinkDevice/Private/StereoVideoLiveLinkDeviceModule.cpp
deleted file mode 100644
diff --git a/Engine/Plugins/VirtualProduction/CaptureManager/CaptureManagerDevices/Source/StereoVideoLiveLinkDevice/StereoVideoLiveLinkDevice.Build.cs b/Engine/Plugins/VirtualProduction/CaptureManager/CaptureManagerDevices/Source/StereoVideoLiveLinkDevice/StereoVideoLiveLinkDevice.Build.cs
deleted file mode 100644
diff --git a/Engine/Plugins/VirtualProduction/CaptureManager/CaptureManagerDevices/Source/TakeArchiveIngestDevice/Private/TakeArchiveIngestDevice.cpp b/Engine/Plugins/VirtualProduction/CaptureManager/CaptureManagerDevices/Source/TakeArchiveIngestDevice/Private/TakeArchiveIngestDevice.cpp
new file mode 100644
--- /dev/null
+++ b/Engine/Plugins/VirtualProduction/CaptureManager/CaptureManagerDevices/Source/TakeArchiveIngestDevice/Private/TakeArchiveIngestDevice.cpp
@@ -0,0 +1,308 @@
+// Copyright Epic Games, Inc. All Rights Reserved.
+
+#include "TakeArchiveIngestDevice.h"
+
+#include "HAL/FileManager.h"
+#include "Settings/CaptureManagerSettings.h"
+
+#include "Utils/TakeArchiveIngestDeviceLog.h"
+
+#include "LiveLinkFaceMetadata.h"
+#include "StereoCameraMetadataParseUtils.h"
+
+#include "Utils/IngestLiveLinkDeviceUtils.h"
+#include "Utils/CaptureExtractTimecode.h"
+#include "Utils/ParseTakeUtils.h"
+
+DEFINE_LOG_CATEGORY(LogTakeArchiveIngestDevice);
+
+const UTakeArchiveIngestDeviceSettings* UTakeArchiveIngestDevice::GetSettings() const
+{
+	return GetDeviceSettings<UTakeArchiveIngestDeviceSettings>();
+}
+
+TSubclassOf<ULiveLinkDeviceSettings> UTakeArchiveIngestDevice::GetSettingsClass() const
+{
+	return UTakeArchiveIngestDeviceSettings::StaticClass();
+}
+
+FText UTakeArchiveIngestDevice::GetDisplayName() const
+{
+	return FText::FromString(GetSettings()->DisplayName);
+}
+
+EDeviceHealth UTakeArchiveIngestDevice::GetDeviceHealth() const
+{
+	return EDeviceHealth::Nominal;
+}
+
+FText UTakeArchiveIngestDevice::GetHealthText() const
+{
+	return FText::FromString("Nominal");
+}
+
+void UTakeArchiveIngestDevice::OnSettingChanged(const FPropertyChangedEvent& InPropertyChangedEvent)
+{
+	Super::OnSettingChanged(InPropertyChangedEvent);
+
+	if (GET_MEMBER_NAME_CHECKED(UTakeArchiveIngestDeviceSettings, TakeDirectory) == InPropertyChangedEvent.GetMemberPropertyName())
+	{
+		SetConnectionStatus(ELiveLinkDeviceConnectionStatus::Disconnected);
+
+		FString Path = GetSettings()->TakeDirectory.Path;
+		if (!Path.IsEmpty() && FPaths::DirectoryExists(Path))
+		{
+			SetConnectionStatus(ELiveLinkDeviceConnectionStatus::Connected);
+		}
+		else
+		{
+			SetConnectionStatus(ELiveLinkDeviceConnectionStatus::Disconnected);
+		}
+	}
+}
+
+FString UTakeArchiveIngestDevice::GetFullTakePath(UE::CaptureManager::FTakeId InTakeId) const
+{
+	FString StoragePath = GetSettings()->TakeDirectory.Path;
+
+	if (const FString* RelativeTakePath = RelativeTakePaths.Find(InTakeId))
+	{
+		return FPaths::Combine(StoragePath, *RelativeTakePath);
+	}
+
+	return FString();
+}
+
+void UTakeArchiveIngestDevice::RunConvertAndUploadTake(const UIngestCapability_ProcessHandle* InProcessHandle, const UIngestCapability_Options* InIngestOptions)
+{
+	AsyncTask(ENamedThreads::AnyThread, [this,
+			  ProcessHandle = TStrongObjectPtr<const UIngestCapability_ProcessHandle>(InProcessHandle),
+			  IngestOptions = TStrongObjectPtr<const UIngestCapability_Options>(InIngestOptions)]()
+	{
+		static constexpr uint32 NumberOfTasks = 2; // Convert, Upload
+
+		using namespace UE::CaptureManager;
+		TSharedPtr<FTaskProgress> TaskProgress
+			= MakeShared<FTaskProgress>(NumberOfTasks, FTaskProgress::FProgressReporter::CreateLambda([this, ProcessHandle](double InProgress)
+		{
+			ExecuteProcessProgressReporter(ProcessHandle.Get(), InProgress);
+		}));
+
+		Super::IngestTake(ProcessHandle.Get(), IngestOptions.Get(), MoveTemp(TaskProgress));
+	});
+}
+
+void UTakeArchiveIngestDevice::UpdateTakeList_Implementation(UIngestCapability_UpdateTakeListCallback* InCallback)
+{
+	AsyncTask(ENamedThreads::AnyThread, [this, Callback = TStrongObjectPtr<UIngestCapability_UpdateTakeListCallback>(InCallback)]()
+	{
+		RemoveAllTakes();
+
+		FString StoragePath = GetSettings()->TakeDirectory.Path;
+
+		TArray<FString> TakeMetadataFiles;
+		const bool bIterationResult = IFileManager::Get().IterateDirectoryRecursively(*StoragePath, [&](const TCHAR* InFileNameOrDirectory, bool bInIsDirectory)
+		{
+			if (!bInIsDirectory)
+			{
+				const FString CurrentFileName = FPaths::GetCleanFilename(InFileNameOrDirectory);
+				if (CurrentFileName == TEXT("take.json"))
+				{
+					TakeMetadataFiles.Add(InFileNameOrDirectory);
+				}
+				else if (FPaths::GetExtension(CurrentFileName) == FTakeMetadata::FileExtension)
+				{
+					TakeMetadataFiles.Add(InFileNameOrDirectory);
+				}
+			}
+
+			return true;
+		});
+
+		for (const FString& CurrentTakeFile : TakeMetadataFiles)
+		{
+			TOptional<FTakeMetadata> TakeInfoResult = ReadTake(CurrentTakeFile);
+
+			if (TakeInfoResult.IsSet())
+			{
+				int32 TakeId = AddTake(TakeInfoResult.GetValue());
+				FString CurrentDirectory = FPaths::GetPath(CurrentTakeFile);
+
+				FString TakePath = CurrentDirectory.Right(CurrentDirectory.Len() - StoragePath.Len());
+				FPaths::NormalizeDirectoryName(TakePath); // Removes trailing slash
+				RelativeTakePaths.Add(TakeId, MoveTemp(TakePath));
+
+				PublishEvent<FIngestCapability_TakeAddedEvent>(TakeId);
+			}
+		}
+
+		ExecuteUpdateTakeListCallback(Callback.Get(), Execute_GetTakeIdentifiers(this));
+	});
+}
+
+ELiveLinkDeviceConnectionStatus UTakeArchiveIngestDevice::GetConnectionStatus_Implementation() const
+{
+	const UTakeArchiveIngestDeviceSettings* DeviceSettings = GetSettings();
+	FString Path = DeviceSettings->TakeDirectory.Path;
+
+	if (!Path.IsEmpty() && FPaths::DirectoryExists(Path))
+	{
+		return ELiveLinkDeviceConnectionStatus::Connected;
+	}
+
+	return ELiveLinkDeviceConnectionStatus::Disconnected;
+}
+
+FString UTakeArchiveIngestDevice::GetHardwareId_Implementation() const
+{
+	return FPlatformMisc::GetDeviceId();
+}
+
+bool UTakeArchiveIngestDevice::SetHardwareId_Implementation(const FString& HardwareID)
+{
+	return false;
+}
+
+bool UTakeArchiveIngestDevice::Connect_Implementation()
+{
+	const UTakeArchiveIngestDeviceSettings* DeviceSettings = GetSettings();
+	FString Path = DeviceSettings->TakeDirectory.Path;
+
+	if (!Path.IsEmpty() && FPaths::DirectoryExists(Path))
+	{
+		SetConnectionStatus(ELiveLinkDeviceConnectionStatus::Connected);
+		return true;
+	}
+	
+	return false;
+}
+
+bool UTakeArchiveIngestDevice::Disconnect_Implementation()
+{
+	SetConnectionStatus(ELiveLinkDeviceConnectionStatus::Disconnected);
+	return true;
+}
+
+TOptional<FTakeMetadata> UTakeArchiveIngestDevice::ReadTake(const FString& InCurrentTakeFile)
+{
+	using namespace UE::CaptureManager;
+
+	FTakeMetadataParser TakeMetadataParser;
+	TValueOrError<FTakeMetadata, FTakeMetadataParserError> TakeMetadataResult = TakeMetadataParser.Parse(InCurrentTakeFile);
+	FString CurrentDirectory = FPaths::GetPath(InCurrentTakeFile);
+
+	if (TakeMetadataResult.HasError())
+	{
+		FTakeMetadataParserError ParserError = TakeMetadataResult.StealError();
+		FString OriginName = UE::CaptureManager::ErrorOriginToString(ParserError.Origin);
+
+		FString Message =
+			FString::Format(TEXT("Unable to parse take metadata file - {0} (Error origin: {1}): {2}"),
+							{
+								InCurrentTakeFile,
+								ParserError.Message.ToString(),
+								OriginName
+							});
+
+		UE_LOG(LogTakeArchiveIngestDevice, Warning, TEXT("%s"), *Message);
+
+		UE_LOG(LogTakeArchiveIngestDevice, Display, TEXT("Checking backwards compatible take metadata formats"), *CurrentDirectory);
+
+		TArray<FText> ValidationFailures;
+
+		UE_LOG(LogTakeArchiveIngestDevice, Display, TEXT("Checking directory (%s) for pre UE 5.6 LiveLink take metadata format"), *CurrentDirectory);
+		TOptional<FTakeMetadata> ParseOldMetadataResult = LiveLinkMetadata::ParseOldLiveLinkTakeMetadata(CurrentDirectory, ValidationFailures);
+
+		if (ParseOldMetadataResult.IsSet())
+		{
+			return ParseOldMetadataResult.GetValue();
+		}
+
+		UE_LOG(LogTakeArchiveIngestDevice, Display, TEXT("Checking directory (%s) for pre UE 5.6 StereoCamera take metadata format"), *CurrentDirectory);
+		ParseOldMetadataResult = StereoCameraMetadata::ParseOldStereoCameraMetadata(CurrentDirectory, ValidationFailures);
+
+		if (ParseOldMetadataResult.IsSet())
+		{
+			return ParseOldMetadataResult.GetValue();
+		}
+
+		UE_LOG(LogTakeArchiveIngestDevice, Error, TEXT("Failed to parse take metadata file: %s"), *CurrentDirectory);
+
+		return {};
+	}
+
+	FTakeMetadata& TakeMetadata = TakeMetadataResult.GetValue();
+
+	TOptional<FString> ThumbnailPath = TakeMetadata.Thumbnail.GetThumbnailPath();
+	if (ThumbnailPath.IsSet())
+	{
+		if (FPaths::IsRelative(*ThumbnailPath))
+		{
+			TakeMetadata.Thumbnail = FTakeThumbnailData(FPaths::ConvertRelativePathToFull(CurrentDirectory, *ThumbnailPath));
+		}
+	}
+	
+	bool bVideoFrameRateSet = false;
+	FFrameRate VideoFrameRate;
+	for (FTakeMetadata::FVideo& Video : TakeMetadata.Video)
+	{
+		if (FPaths::IsRelative(*Video.Path))
+		{
+			Video.Path = FPaths::ConvertRelativePathToFull(CurrentDirectory, *Video.Path);
+		}
+
+		// Check if this is a file (and not an image sequence directory)
+		if (FPaths::FileExists(Video.Path))
+		{
+			TSharedPtr<FCaptureExtractVideoTimecode> Extractor = MakeShareable(new FCaptureExtractVideoTimecode(Video.Path));
+			FTimecodeAndRateResult TimecodeAndRateResult = Extractor->Extract();
+			if (TimecodeAndRateResult.IsValid())
+			{
+				FTimecodeAndRate TimecodeAndRate = TimecodeAndRateResult.GetValue();
+				if (!Video.TimecodeStart.IsSet())
+				{
+					Video.TimecodeStart = TimecodeAndRate.Timecode.ToString();
+				}
+			}
+		}
+
+		if (!bVideoFrameRateSet)
+		{
+			VideoFrameRate = UE::CaptureManager::ParseFrameRate(Video.FrameRate);
+		}
+	}
+
+	for (FTakeMetadata::FAudio& Audio : TakeMetadata.Audio)
+	{
+		if (FPaths::IsRelative(*Audio.Path))
+		{
+			Audio.Path = FPaths::ConvertRelativePathToFull(CurrentDirectory, *Audio.Path);
+		}
+
+		TSharedPtr<FCaptureExtractAudioTimecode> Extractor = MakeShareable(new FCaptureExtractAudioTimecode(Audio.Path));
+
+		// The video frame rate will be used to calculate the timecode rate if the timecode rate cannot be extracted from the audio file
+		// If the VideoFrameRate is not set, the timecode rate may be set to the audio sample rate
+		FTimecodeAndRateResult TimecodeAndRateResult = Extractor->Extract(VideoFrameRate);
+		if (TimecodeAndRateResult.IsValid())
+		{
+			FTimecodeAndRate TimecodeAndRate = TimecodeAndRateResult.GetValue();
+			Audio.TimecodeStart = TimecodeAndRate.Timecode.ToString();
+			
+			if (TimecodeAndRate.MaybeTimecodeRate.IsSet())
+			{
+				Audio.TimecodeRate = static_cast<float>(TimecodeAndRate.MaybeTimecodeRate.GetValue().AsDecimal());
+			}
+		}
+	}
+	
+	for (FTakeMetadata::FCalibration& Calib : TakeMetadata.Calibration)
+	{
+		if (FPaths::IsRelative(*Calib.Path))
+		{
+			Calib.Path = FPaths::ConvertRelativePathToFull(CurrentDirectory, *Calib.Path);
+		}
+	}
+
+	return TakeMetadataResult.StealValue();
+}
\ No newline at end of file
diff --git a/Engine/Plugins/VirtualProduction/CaptureManager/CaptureManagerDevices/Source/TakeArchiveIngestDevice/Private/TakeArchiveIngestDevice.h b/Engine/Plugins/VirtualProduction/CaptureManager/CaptureManagerDevices/Source/TakeArchiveIngestDevice/Private/TakeArchiveIngestDevice.h
new file mode 100644
--- /dev/null
+++ b/Engine/Plugins/VirtualProduction/CaptureManager/CaptureManagerDevices/Source/TakeArchiveIngestDevice/Private/TakeArchiveIngestDevice.h
@@ -0,0 +1,72 @@
+// Copyright Epic Games, Inc. All Rights Reserved.
+
+#pragma once
+
+#include "LiveLinkDevice.h"
+
+#include "Ingest/LiveLinkDeviceCapability_Ingest.h"
+#include "LiveLinkDeviceCapability_Connection.h"
+
+#include "Misc/Optional.h"
+
+#include "Async/TaskProgress.h"
+
+#include "BaseIngestLiveLinkDevice.h"
+
+#include "TakeArchiveIngestDevice.generated.h"
+
+
+UCLASS(BlueprintType)
+class UTakeArchiveIngestDeviceSettings : public ULiveLinkDeviceSettings
+{
+	GENERATED_BODY()
+
+public:
+	UPROPERTY(EditAnywhere, BlueprintReadWrite, Category = "Take Archive Ingest")
+	FString DisplayName = TEXT("Take Archive Ingest");
+
+	/* Path to a directory containing the take(s) data */
+	UPROPERTY(EditAnywhere, BlueprintReadWrite, Category = "Take Archive Ingest")
+	FDirectoryPath TakeDirectory;
+};
+
+
+UCLASS(BlueprintType, meta = (DisplayName = "Take Archive Ingest", ToolTip = "Ingest take archives (described by .cptake files) and legacy Capture Manager takes"))
+class UTakeArchiveIngestDevice final : public UBaseIngestLiveLinkDevice
+	, public ILiveLinkDeviceCapability_Connection
+{
+	GENERATED_BODY()
+
+public:
+	UFUNCTION(BlueprintCallable, Category = "FilesystemDeviceSettings")
+	const UTakeArchiveIngestDeviceSettings* GetSettings() const;
+
+	//~ Begin ULiveLinkDevice interface
+	virtual TSubclassOf<ULiveLinkDeviceSettings> GetSettingsClass() const override;
+	virtual FText GetDisplayName() const override;
+	virtual EDeviceHealth GetDeviceHealth() const override;
+	virtual FText GetHealthText() const override;
+
+	virtual void OnSettingChanged(const FPropertyChangedEvent& InPropertyChangedEvent) override;
+	//~ End ULiveLinkDevice interface
+
+private:
+	virtual FString GetFullTakePath(UE::CaptureManager::FTakeId InTakeId) const override;
+
+	//~ Begin ULiveLinkDeviceCapability_Ingest interface
+	virtual void RunConvertAndUploadTake(const UIngestCapability_ProcessHandle* InProcessHandle, const UIngestCapability_Options* InIngestOptions) override;
+	virtual void UpdateTakeList_Implementation(UIngestCapability_UpdateTakeListCallback* InCallback) override;
+	//~ End ULiveLinkDeviceCapability_Ingest interface
+
+	//~ Begin ILiveLinkDeviceCapability_Connection interface
+	virtual ELiveLinkDeviceConnectionStatus GetConnectionStatus_Implementation() const override;
+	virtual FString GetHardwareId_Implementation() const override;
+	virtual bool SetHardwareId_Implementation(const FString& HardwareID) override;
+	virtual bool Connect_Implementation() override;
+	virtual bool Disconnect_Implementation() override;
+	//~ End ILiveLinkDeviceCapability_Connection interface
+
+	static TOptional<FTakeMetadata> ReadTake(const FString& InCurrentTakeFile);
+
+	TMap<int32, FString> RelativeTakePaths;
+};
diff --git a/Engine/Plugins/VirtualProduction/CaptureManager/CaptureManagerDevices/Source/TakeArchiveIngestDevice/Private/TakeArchiveIngestDeviceModule.cpp b/Engine/Plugins/VirtualProduction/CaptureManager/CaptureManagerDevices/Source/TakeArchiveIngestDevice/Private/TakeArchiveIngestDeviceModule.cpp
new file mode 100644
--- /dev/null
+++ b/Engine/Plugins/VirtualProduction/CaptureManager/CaptureManagerDevices/Source/TakeArchiveIngestDevice/Private/TakeArchiveIngestDeviceModule.cpp
@@ -0,0 +1,5 @@
+// Copyright Epic Games, Inc. All Rights Reserved.
+
+#include "Modules/ModuleManager.h"
+
+IMPLEMENT_MODULE(FDefaultModuleImpl, TakeArchiveIngestDevice)
\ No newline at end of file
diff --git a/Engine/Plugins/VirtualProduction/CaptureManager/CaptureManagerDevices/Source/TakeArchiveIngestDevice/Private/Utils/TakeArchiveIngestDeviceLog.h b/Engine/Plugins/VirtualProduction/CaptureManager/CaptureManagerDevices/Source/TakeArchiveIngestDevice/Private/Utils/TakeArchiveIngestDeviceLog.h
new file mode 100644
--- /dev/null
+++ b/Engine/Plugins/VirtualProduction/CaptureManager/CaptureManagerDevices/Source/TakeArchiveIngestDevice/Private/Utils/TakeArchiveIngestDeviceLog.h
@@ -0,0 +1,7 @@
+// Copyright Epic Games, Inc. All Rights Reserved.
+
+#pragma once
+
+#include "Logging/LogMacros.h"
+
+DECLARE_LOG_CATEGORY_EXTERN(LogTakeArchiveIngestDevice, Log, All);
diff --git a/Engine/Plugins/VirtualProduction/CaptureManager/CaptureManagerDevices/Source/TakeArchiveIngestDevice/TakeArchiveIngestDevice.Build.cs b/Engine/Plugins/VirtualProduction/CaptureManager/CaptureManagerDevices/Source/TakeArchiveIngestDevice/TakeArchiveIngestDevice.Build.cs
new file mode 100644
--- /dev/null
+++ b/Engine/Plugins/VirtualProduction/CaptureManager/CaptureManagerDevices/Source/TakeArchiveIngestDevice/TakeArchiveIngestDevice.Build.cs
@@ -0,0 +1,44 @@
+// Copyright Epic Games, Inc. All Rights Reserved.
+
+using UnrealBuildTool;
+using System.IO;
+
+public class TakeArchiveIngestDevice : ModuleRules
+{
+	public TakeArchiveIngestDevice(ReadOnlyTargetRules Target) : base(Target)
+	{
+		PCHUsage = PCHUsageMode.UseExplicitOrSharedPCHs;
+		ShortName = "TakeArchiveIngest";
+
+		PublicDependencyModuleNames.AddRange(new string[]
+		{
+			"CaptureUtils",
+			"Core"
+		});
+
+		PrivateDependencyModuleNames.AddRange(new string[]
+		{
+			"Engine",
+			"Json",
+			"CaptureManagerMediaRW",
+			"IngestLiveLinkDevice",
+			"DataIngestCore",
+			"LiveLinkDevice",
+			"LiveLinkCapabilities",
+			"Media",
+			"LiveLinkHubCaptureMessaging",
+			"LiveLinkHub",
+			"LiveLinkFaceMetadata",
+			"ToolWidgets",
+			"Slate",
+			"SlateCore",
+			"StereoCameraMetadata",
+			"CaptureManagerTakeMetadata",
+			"CaptureManagerSettings",
+			"CoreUObject",
+			"UnrealEd"
+		});
+
+		CppCompileWarningSettings.UnsafeTypeCastWarningLevel = WarningLevel.Error;
+	}
+}
