Timecode extraction is taking too long

diff --git a/Engine/Plugins/VirtualProduction/CaptureManager/CaptureManagerDevices/Source/TakeArchiveIngestDevice/Private/TakeArchiveIngestDevice.cpp b/Engine/Plugins/VirtualProduction/CaptureManager/CaptureManagerDevices/Source/TakeArchiveIngestDevice/Private/TakeArchiveIngestDevice.cpp
--- a/Engine/Plugins/VirtualProduction/CaptureManager/CaptureManagerDevices/Source/TakeArchiveIngestDevice/Private/TakeArchiveIngestDevice.cpp
+++ b/Engine/Plugins/VirtualProduction/CaptureManager/CaptureManagerDevices/Source/TakeArchiveIngestDevice/Private/TakeArchiveIngestDevice.cpp
@@ -252,17 +252,14 @@ TOptional<FTakeMetadata> UTakeArchiveIngestDevice::ReadTake(const FString& InCur
 		}
 
 		// Check if this is a file (and not an image sequence directory)
-		if (FPaths::FileExists(Video.Path))
+		if (FPaths::FileExists(Video.Path) && !Video.TimecodeStart.IsSet())
 		{
 			TSharedPtr<FCaptureExtractVideoTimecode> Extractor = MakeShareable(new FCaptureExtractVideoTimecode(Video.Path));
 			FTimecodeAndRateResult TimecodeAndRateResult = Extractor->Extract();
 			if (TimecodeAndRateResult.IsValid())
 			{
 				FTimecodeAndRate TimecodeAndRate = TimecodeAndRateResult.GetValue();
-				if (!Video.TimecodeStart.IsSet())
-				{
-					Video.TimecodeStart = TimecodeAndRate.Timecode.ToString();
-				}
+				Video.TimecodeStart = TimecodeAndRate.Timecode.ToString();
 			}
 		}
 
