Make parts of CPSLiveLinkDevice public

diff --git a/Engine/Plugins/VirtualProduction/CaptureManager/CaptureManagerApp/Source/LiveLinkCapabilities/LiveLinkCapabilities.Build.cs b/Engine/Plugins/VirtualProduction/CaptureManager/CaptureManagerApp/Source/LiveLinkCapabilities/LiveLinkCapabilities.Build.cs
--- a/Engine/Plugins/VirtualProduction/CaptureManager/CaptureManagerApp/Source/LiveLinkCapabilities/LiveLinkCapabilities.Build.cs
+++ b/Engine/Plugins/VirtualProduction/CaptureManager/CaptureManagerApp/Source/LiveLinkCapabilities/LiveLinkCapabilities.Build.cs
@@ -14,7 +14,8 @@ public class LiveLinkCapabilities : ModuleRules
 		{
 			"CaptureUtils",
 			"Core",
-			"CaptureManagerUnrealEndpoint" // TODO: Only needed for the ULiveLinkHubIngestClientInfo definition, move it and remove this dependency
+			"CaptureManagerUnrealEndpoint", // TODO: Only needed for the ULiveLinkHubIngestClientInfo definition, move it and remove this dependency
+			"CaptureManagerTakeMetadata"
 		});
 
 		PrivateDependencyModuleNames.AddRange(new string[]
@@ -26,7 +27,6 @@ public class LiveLinkCapabilities : ModuleRules
 			"LiveLinkHubCaptureMessaging",
 			"SlateCore",
 			"Slate",
-			"CaptureManagerTakeMetadata",
 			"CoreUObject"
 		});
 
diff --git a/Engine/Plugins/VirtualProduction/CaptureManager/CaptureManagerDevices/Source/CPSLiveLinkDevice/Private/Customizations/DeviceIpAddressCustomization.cpp b/Engine/Plugins/VirtualProduction/CaptureManager/CaptureManagerDevices/Source/CPSLiveLinkDevice/Private/Customizations/DeviceIpAddressCustomization.cpp
--- a/Engine/Plugins/VirtualProduction/CaptureManager/CaptureManagerDevices/Source/CPSLiveLinkDevice/Private/Customizations/DeviceIpAddressCustomization.cpp
+++ b/Engine/Plugins/VirtualProduction/CaptureManager/CaptureManagerDevices/Source/CPSLiveLinkDevice/Private/Customizations/DeviceIpAddressCustomization.cpp
@@ -1,6 +1,6 @@
 // Copyright Epic Games, Inc. All Rights Reserved.
 
-#include "DeviceIpAddressCustomization.h"
+#include "Customizations/DeviceIpAddressCustomization.h"
 #include "Widgets/Input/SEditableTextBox.h"
 #include "Editor.h"
 
diff --git a/Engine/Plugins/VirtualProduction/CaptureManager/CaptureManagerDevices/Source/CPSLiveLinkDevice/Private/Customizations/DeviceIpAddressCustomization.h b/Engine/Plugins/VirtualProduction/CaptureManager/CaptureManagerDevices/Source/CPSLiveLinkDevice/Private/Customizations/DeviceIpAddressCustomization.h
deleted file mode 100644
diff --git a/Engine/Plugins/VirtualProduction/CaptureManager/CaptureManagerDevices/Source/CPSLiveLinkDevice/Private/Customizations/ToggleConnectActionCustomization.h b/Engine/Plugins/VirtualProduction/CaptureManager/CaptureManagerDevices/Source/CPSLiveLinkDevice/Private/Customizations/ToggleConnectActionCustomization.h
deleted file mode 100644
diff --git a/Engine/Plugins/VirtualProduction/CaptureManager/CaptureManagerDevices/Source/CPSLiveLinkDevice/Private/LiveLinkFaceDevice.h b/Engine/Plugins/VirtualProduction/CaptureManager/CaptureManagerDevices/Source/CPSLiveLinkDevice/Private/LiveLinkFaceDevice.h
deleted file mode 100644
diff --git a/Engine/Plugins/VirtualProduction/CaptureManager/CaptureManagerDevices/Source/CPSLiveLinkDevice/Private/Protocol/CPSDataStream.cpp b/Engine/Plugins/VirtualProduction/CaptureManager/CaptureManagerDevices/Source/CPSLiveLinkDevice/Private/Protocol/CPSDataStream.cpp
--- a/Engine/Plugins/VirtualProduction/CaptureManager/CaptureManagerDevices/Source/CPSLiveLinkDevice/Private/Protocol/CPSDataStream.cpp
+++ b/Engine/Plugins/VirtualProduction/CaptureManager/CaptureManagerDevices/Source/CPSLiveLinkDevice/Private/Protocol/CPSDataStream.cpp
@@ -1,6 +1,6 @@
 // Copyright Epic Games, Inc. All Rights Reserved.
 
-#include "CPSDataStream.h"
+#include "Protocol/CPSDataStream.h"
 
 #include "Misc/FileHelper.h"
 #include "Misc/Paths.h"
diff --git a/Engine/Plugins/VirtualProduction/CaptureManager/CaptureManagerDevices/Source/CPSLiveLinkDevice/Private/Protocol/CPSDataStream.h b/Engine/Plugins/VirtualProduction/CaptureManager/CaptureManagerDevices/Source/CPSLiveLinkDevice/Private/Protocol/CPSDataStream.h
deleted file mode 100644
diff --git a/Engine/Plugins/VirtualProduction/CaptureManager/CaptureManagerDevices/Source/CPSLiveLinkDevice/Private/Protocol/CPSDevice.cpp b/Engine/Plugins/VirtualProduction/CaptureManager/CaptureManagerDevices/Source/CPSLiveLinkDevice/Private/Protocol/CPSDevice.cpp
--- a/Engine/Plugins/VirtualProduction/CaptureManager/CaptureManagerDevices/Source/CPSLiveLinkDevice/Private/Protocol/CPSDevice.cpp
+++ b/Engine/Plugins/VirtualProduction/CaptureManager/CaptureManagerDevices/Source/CPSLiveLinkDevice/Private/Protocol/CPSDevice.cpp
@@ -1,6 +1,6 @@
 // Copyright Epic Games, Inc. All Rights Reserved.
 
-#include "CPSDevice.h"
+#include "Protocol/CPSDevice.h"
 
 #include "Control/Messages/Constants.h"
 
diff --git a/Engine/Plugins/VirtualProduction/CaptureManager/CaptureManagerDevices/Source/CPSLiveLinkDevice/Private/Protocol/CPSDevice.h b/Engine/Plugins/VirtualProduction/CaptureManager/CaptureManagerDevices/Source/CPSLiveLinkDevice/Private/Protocol/CPSDevice.h
deleted file mode 100644
diff --git a/Engine/Plugins/VirtualProduction/CaptureManager/CaptureManagerDevices/Source/CPSLiveLinkDevice/Private/Protocol/CPSFileStream.cpp b/Engine/Plugins/VirtualProduction/CaptureManager/CaptureManagerDevices/Source/CPSLiveLinkDevice/Private/Protocol/CPSFileStream.cpp
--- a/Engine/Plugins/VirtualProduction/CaptureManager/CaptureManagerDevices/Source/CPSLiveLinkDevice/Private/Protocol/CPSFileStream.cpp
+++ b/Engine/Plugins/VirtualProduction/CaptureManager/CaptureManagerDevices/Source/CPSLiveLinkDevice/Private/Protocol/CPSFileStream.cpp
@@ -1,6 +1,6 @@
 // Copyright Epic Games, Inc. All Rights Reserved.
 
-#include "CPSFileStream.h"
+#include "Protocol/CPSFileStream.h"
 
 #include "Misc/FileHelper.h"
 #include "Misc/Paths.h"
diff --git a/Engine/Plugins/VirtualProduction/CaptureManager/CaptureManagerDevices/Source/CPSLiveLinkDevice/Private/Protocol/CPSFileStream.h b/Engine/Plugins/VirtualProduction/CaptureManager/CaptureManagerDevices/Source/CPSLiveLinkDevice/Private/Protocol/CPSFileStream.h
deleted file mode 100644
diff --git a/Engine/Plugins/VirtualProduction/CaptureManager/CaptureManagerDevices/Source/CPSLiveLinkDevice/Public/Customizations/DeviceIpAddressCustomization.h b/Engine/Plugins/VirtualProduction/CaptureManager/CaptureManagerDevices/Source/CPSLiveLinkDevice/Public/Customizations/DeviceIpAddressCustomization.h
new file mode 100644
--- /dev/null
+++ b/Engine/Plugins/VirtualProduction/CaptureManager/CaptureManagerDevices/Source/CPSLiveLinkDevice/Public/Customizations/DeviceIpAddressCustomization.h
@@ -0,0 +1,38 @@
+// Copyright Epic Games, Inc. All Rights Reserved.
+
+#pragma once
+
+#include "DetailLayoutBuilder.h"
+#include "DetailWidgetRow.h"
+#include "IDetailChildrenBuilder.h"
+
+#include "DeviceIpAddressCustomization.generated.h"
+
+USTRUCT(BlueprintType)
+struct CPSLIVELINKDEVICE_API FDeviceIpAddress
+{
+	GENERATED_BODY()
+
+	UPROPERTY(EditAnywhere, BlueprintReadWrite, Category = DeviceAddress, DisplayName="Ip Address")
+	FString IpAddressString;
+};
+
+class CPSLIVELINKDEVICE_API FDeviceIpAddressCustomization : public IPropertyTypeCustomization
+{
+public:
+	FDeviceIpAddressCustomization();
+
+	virtual void CustomizeHeader(TSharedRef<IPropertyHandle> PropertyHandle, FDetailWidgetRow& HeaderRow, IPropertyTypeCustomizationUtils& CustomizationUtils) override;
+	virtual void CustomizeChildren(TSharedRef<IPropertyHandle> PropertyHandle, IDetailChildrenBuilder& ChildBuilder, IPropertyTypeCustomizationUtils& CustomizationUtils) override;
+
+	static TOptional<FText> VerifyIpAddress(const FString& InIpAddress, TRange<int32> InIpAddressRange);
+
+private:
+	FText OnGetDeviceIpAddress() const;
+	bool OnDeviceIpAddressVerify(const FText& InText, FText& OutErrorText);
+	void OnDeviceIpAddressCommited(const FText& InText, ETextCommit::Type CommitInfo);
+	bool IsReadOnly() const;
+
+	TSharedPtr<IPropertyHandle> DeviceIpAddressProperty;
+	TRange<int32> IpAddressRange;
+};
diff --git a/Engine/Plugins/VirtualProduction/CaptureManager/CaptureManagerDevices/Source/CPSLiveLinkDevice/Public/Customizations/ToggleConnectActionCustomization.h b/Engine/Plugins/VirtualProduction/CaptureManager/CaptureManagerDevices/Source/CPSLiveLinkDevice/Public/Customizations/ToggleConnectActionCustomization.h
new file mode 100644
--- /dev/null
+++ b/Engine/Plugins/VirtualProduction/CaptureManager/CaptureManagerDevices/Source/CPSLiveLinkDevice/Public/Customizations/ToggleConnectActionCustomization.h
@@ -0,0 +1,49 @@
+// Copyright Epic Games, Inc. All Rights Reserved.
+
+#pragma once
+
+#include "IPropertyTypeCustomization.h"
+#include "Input/Reply.h"
+
+#include "LiveLinkDevice.h"
+#include "LiveLinkDeviceCapability_Connection.h"
+
+#include "ToggleConnectActionCustomization.generated.h"
+
+USTRUCT()
+struct CPSLIVELINKDEVICE_API FToggleConnectAction
+{
+	GENERATED_BODY()
+
+	UPROPERTY()
+	FGuid DeviceGuid;
+};
+
+class CPSLIVELINKDEVICE_API FToggleConnectActionCustomization : public IPropertyTypeCustomization
+{
+public:
+
+	virtual void CustomizeHeader(TSharedRef<IPropertyHandle> InPropertyHandle, 
+								 FDetailWidgetRow& InHeaderRow, 
+								 IPropertyTypeCustomizationUtils& InCustomizationUtils) override;
+	virtual void CustomizeChildren(TSharedRef<IPropertyHandle> InPropertyHandle, 
+								   IDetailChildrenBuilder& InChildBuilder, 
+								   IPropertyTypeCustomizationUtils& InCustomizationUtils) override;
+private:
+
+	void OnConnectionStatusChanged(ELiveLinkDeviceConnectionStatus InStatus);
+
+	FText GetButtonName() const;
+	FReply OnConnectButtonToggled();
+	bool IsConnectButtonEnabled() const;
+
+	EVisibility GetStopActionVisibility() const;
+	FReply OnStopActionClicked();
+	FText GetStopActionTooltip() const;
+
+	TObjectPtr<ULiveLinkDevice> GetDevice() const;
+
+	TSharedPtr<IPropertyHandle> PropertyHandle;
+	TObjectPtr<ULiveLinkDevice> Device;
+	std::atomic_bool bIsDisconnecting = false;
+};
\ No newline at end of file
diff --git a/Engine/Plugins/VirtualProduction/CaptureManager/CaptureManagerDevices/Source/CPSLiveLinkDevice/Public/LiveLinkFaceDevice.h b/Engine/Plugins/VirtualProduction/CaptureManager/CaptureManagerDevices/Source/CPSLiveLinkDevice/Public/LiveLinkFaceDevice.h
new file mode 100644
--- /dev/null
+++ b/Engine/Plugins/VirtualProduction/CaptureManager/CaptureManagerDevices/Source/CPSLiveLinkDevice/Public/LiveLinkFaceDevice.h
@@ -0,0 +1,117 @@
+// Copyright Epic Games, Inc. All Rights Reserved.
+
+#pragma once
+
+#include "LiveLinkDevice.h"
+
+#include "Ingest/LiveLinkDeviceCapability_Ingest.h"
+#include "LiveLinkDeviceCapability_Connection.h"
+#include "LiveLinkDeviceCapability_Recording.h"
+#include "Misc/Optional.h"
+
+#include "Async/TaskProgress.h"
+
+#include "Protocol/CPSDevice.h"
+#include "Customizations/ToggleConnectActionCustomization.h"
+#include "Customizations/DeviceIpAddressCustomization.h"
+
+#include "BaseIngestLiveLinkDevice.h"
+
+#include "LiveLinkFaceDevice.generated.h"
+
+
+UCLASS(BlueprintType)
+class CPSLIVELINKDEVICE_API ULiveLinkFaceDeviceSettings : public ULiveLinkDeviceSettings
+{
+	GENERATED_BODY()
+
+public:
+	UPROPERTY(EditAnywhere, BlueprintReadWrite, Category = "Live Link Face")
+	FString DisplayName = TEXT("Live Link Face");
+
+	UPROPERTY(EditAnywhere, BlueprintReadWrite, Category = "Live Link Face")
+	FDeviceIpAddress IpAddress;
+
+	UPROPERTY(EditAnywhere, BlueprintReadWrite, Category = "Live Link Face")
+	int32 Port = 14785;
+
+	UPROPERTY(VisibleAnywhere, Category = "Live Link Face")
+	FToggleConnectAction ConnectAction;
+};
+
+
+UCLASS(BlueprintType, meta = (DisplayName = "Live Link Face", ToolTip = "Use for ingest from the Live Link Face app and control recording"))
+class CPSLIVELINKDEVICE_API ULiveLinkFaceDevice : public UBaseIngestLiveLinkDevice
+	, public ILiveLinkDeviceCapability_Connection
+	, public ILiveLinkDeviceCapability_Recording
+{
+	GENERATED_BODY()
+
+public:
+	UFUNCTION(BlueprintCallable, Category = "LiveLinkFaceDeviceSettings")
+	const ULiveLinkFaceDeviceSettings* GetSettings() const;
+
+	//~ Begin ULiveLinkDevice interface
+	virtual TSubclassOf<ULiveLinkDeviceSettings> GetSettingsClass() const override;
+	virtual FText GetDisplayName() const override;
+	virtual EDeviceHealth GetDeviceHealth() const override;
+	virtual FText GetHealthText() const override;
+
+	virtual void OnDeviceAdded() override;
+	virtual void OnDeviceRemoved() override;
+	//~ End ULiveLinkDevice interface
+
+private:
+	virtual FString GetFullTakePath(UE::CaptureManager::FTakeId InTakeId) const override;
+
+	//~ Begin ULiveLinkDeviceCapability_Ingest interface
+	virtual void UpdateTakeList_Implementation(UIngestCapability_UpdateTakeListCallback* InCallback) override;
+	virtual void RunDownloadTake(const UIngestCapability_ProcessHandle* InProcessHandle, const UIngestCapability_Options* InIngestOptions) override;
+	virtual void RunConvertAndUploadTake(const UIngestCapability_ProcessHandle* InProcessHandle, const UIngestCapability_Options* InIngestOptions) override;
+	virtual void CancelIngestProcess_Implementation(const UIngestCapability_ProcessHandle* InProcessHandle) override;
+	//~ End ULiveLinkDeviceCapability_Ingest interface
+
+	//~ Begin ILiveLinkDeviceCapability_Connection interface
+	virtual ELiveLinkDeviceConnectionStatus GetConnectionStatus_Implementation() const override;
+	virtual FString GetHardwareId_Implementation() const override;
+	virtual bool SetHardwareId_Implementation(const FString& HardwareID) override;
+
+protected:
+	virtual bool Connect_Implementation() override;
+	virtual bool Disconnect_Implementation() override;
+	//~ End ILiveLinkDeviceCapability_Connection interface
+
+private:
+	//~ Begin ILiveLinkDeviceCapability_Recording interface
+	virtual bool StartRecording_Implementation() override;
+	virtual bool StopRecording_Implementation() override;
+	virtual bool IsRecording_Implementation() const override;
+	//~ End ILiveLinkDeviceCapability_Recording interface
+
+	void HandleConnectionChanged(TSharedPtr<const UE::CaptureManager::FCaptureEvent> InEvent);
+	void HandleCPSStateUpdate(TSharedPtr<const UE::CaptureManager::FCaptureEvent> InEvent);
+	void HandleCPSEvent(TSharedPtr<const UE::CaptureManager::FCaptureEvent> InEvent);
+
+	void OnExportProgressReport(float InProgress, TStrongObjectPtr<const UIngestCapability_ProcessHandle> InProcessHandle);
+	void OnExportFinished(UE::CaptureManager::TProtocolResult<void> InResult,
+						  FString InTakeName,
+						  TStrongObjectPtr<const UIngestCapability_ProcessHandle> InProcessHandle,
+						  TStrongObjectPtr<const UIngestCapability_Options> InIngestOptions);
+
+	TOptional<FTakeMetadata> ParseTake(const FString& InTakeDirectory, const FString& InTakeName);
+	FTakeMetadata ParseTakeMetadata(const UE::CaptureManager::FGetTakeMetadataResponse::FTakeObject& InTake);
+
+	void FetchPreIngestFiles(TMap<FString, UE::CaptureManager::FTakeId> InNameToIdMap);
+
+	void RemoveDownloadedTakeData(const UE::CaptureManager::FTakeId InTakeId);
+
+	void ExtractTimecodeIfNotSet(FTakeMetadata& InOutTakeMetadata);
+
+	TSharedPtr<UE::CaptureManager::FCPSDevice> Device;
+
+	mutable FCriticalSection DownloadedTakesMutex;
+	TMap<UE::CaptureManager::FTakeId, FString> DownloadedTakes;
+
+	std::atomic_bool bIsRecording = false;
+	std::atomic_bool bIsConnecting = false;
+};
diff --git a/Engine/Plugins/VirtualProduction/CaptureManager/CaptureManagerDevices/Source/CPSLiveLinkDevice/Public/Protocol/CPSDataStream.h b/Engine/Plugins/VirtualProduction/CaptureManager/CaptureManagerDevices/Source/CPSLiveLinkDevice/Public/Protocol/CPSDataStream.h
new file mode 100644
--- /dev/null
+++ b/Engine/Plugins/VirtualProduction/CaptureManager/CaptureManagerDevices/Source/CPSLiveLinkDevice/Public/Protocol/CPSDataStream.h
@@ -0,0 +1,34 @@
+// Copyright Epic Games, Inc. All Rights Reserved.
+
+#pragma once
+
+#include "ExportClient/ExportClient.h"
+
+#include "Containers/Map.h"
+
+namespace UE::CaptureManager {
+
+class CPSLIVELINKDEVICE_API FCPSDataStream : public UE::CaptureManager::FBaseStream
+{
+public:
+	using FData = TArray<uint8>;
+	using FResults = TMap<FString, TMap<FString, UE::CaptureManager::TProtocolResult<FData>>>;
+
+	DECLARE_DELEGATE_OneParam(FFileExportFinished, FResults InExportResults);
+
+	FCPSDataStream(FFileExportFinished InFileExportFinished);
+
+	bool StartFile(const FString& InTakeName, const FString& InFileName) override;
+	bool ProcessData(const FString& InTakeName, const FString& InFileName, const TConstArrayView<uint8>& InData) override;
+	bool FinishFile(const FString& InTakeName, const FString& InFileName, const TStaticArray<uint8, 16>& InHash) override;
+
+	void Finalize(UE::CaptureManager::TProtocolResult<void> InResult) override;
+
+private:
+	FData Data;
+	FFileExportFinished FileExportFinished;
+
+	FResults ExportResults;
+};
+
+}
\ No newline at end of file
diff --git a/Engine/Plugins/VirtualProduction/CaptureManager/CaptureManagerDevices/Source/CPSLiveLinkDevice/Public/Protocol/CPSDevice.h b/Engine/Plugins/VirtualProduction/CaptureManager/CaptureManagerDevices/Source/CPSLiveLinkDevice/Public/Protocol/CPSDevice.h
new file mode 100644
--- /dev/null
+++ b/Engine/Plugins/VirtualProduction/CaptureManager/CaptureManagerDevices/Source/CPSLiveLinkDevice/Public/Protocol/CPSDevice.h
@@ -0,0 +1,150 @@
+// Copyright Epic Games, Inc. All Rights Reserved.
+
+#pragma once
+
+#include "Control/ControlMessenger.h"
+
+#include "ExportClient/ExportClient.h"
+
+#include "Async/EventSourceUtils.h"
+#include "Misc/Optional.h"
+
+#include "Async/CaptureTimerManager.h"
+
+#include "Ingest/IngestCapability_TakeInformation.h"
+
+namespace UE::CaptureManager
+{
+
+struct CPSLIVELINKDEVICE_API FConnectionStateChangedEvent : public FCaptureEvent
+{
+	inline static const FString Name = TEXT("ConnectionStateChanged");
+
+	enum class EState
+	{
+		Unknown = 0,
+		Connecting,
+		Connected,
+		Disconnected
+	};
+
+	FConnectionStateChangedEvent(EState InConnectionState);
+
+	EState ConnectionState;
+};
+
+struct CPSLIVELINKDEVICE_API FCPSEvent : public FCaptureEvent
+{
+	inline static const FString Name = TEXT("CPSEvent");
+
+	FCPSEvent(TSharedPtr<FControlUpdate> InUpdateMessage);
+
+	TSharedPtr<FControlUpdate> UpdateMessage;
+};
+
+struct CPSLIVELINKDEVICE_API FCPSStateEvent : public FCaptureEvent
+{
+	inline static const FString Name = TEXT("CPSStateEvent");
+
+	FCPSStateEvent(FGetStateResponse InGetStateResponse);
+
+	FGetStateResponse GetStateResponse;
+};
+
+class CPSLIVELINKDEVICE_API FCPSDevice :
+	public FCaptureEventSource,
+	public TSharedFromThis<FCPSDevice>
+{
+private:
+
+	struct FPrivateToken
+	{
+		explicit FPrivateToken() = default;
+	};
+
+public:
+
+	static TSharedPtr<FCPSDevice> MakeCPSDevice(FString InDeviceIpAddress, uint16 InDevicePort);
+
+	FCPSDevice(FPrivateToken,
+			   FString InDeviceIpAddress,
+			   uint16 InDevicePort);
+
+	virtual ~FCPSDevice() override;
+
+	void InitiateConnect();
+	void Stop();
+
+	bool IsConnected() const;
+
+	TProtocolResult<void> StartRecording(FString SlateName,
+										 uint16 TakeNumber,
+										 TOptional<FString> Subject = TOptional<FString>(),
+										 TOptional<FString> Scenario = TOptional<FString>(),
+										 TOptional<TArray<FString>> Tags = TOptional<TArray<FString>>());
+
+	TProtocolResult<void> StopRecording();
+
+	TProtocolResult<TArray<FGetTakeMetadataResponse::FTakeObject>> FetchTakeList();
+	TProtocolResult<FGetTakeMetadataResponse::FTakeObject> FetchTake(const FString& InTakeName);
+
+	void AddTakeMetadata(FTakeId InId, FGetTakeMetadataResponse::FTakeObject InTake);
+	void RemoveTakeMetadata(FTakeId InId);
+	FGetTakeMetadataResponse::FTakeObject GetTake(FTakeId InId);
+	FTakeId GetTakeId(const FString& InTakeName);
+
+	void StartExport(FTakeId InTakeId, TUniquePtr<FBaseStream> InStream);
+	void CancelExport(FTakeId InTakeId);
+
+	/** Fetches the thumbnail for a single take. */
+	void FetchThumbnailForTake(FTakeId InTakeId, TUniquePtr<FBaseStream> InStream);
+	/** Fetches the thumbnails for all takes. */
+	void FetchThumbnails(TUniquePtr<FBaseStream> InStream);
+
+	/** Fetches the sepcified file for a single take. */
+	void FetchFileForTake(FTakeId InTakeId, TUniquePtr<FBaseStream> InStream, const FString& InFileName);
+	/** Fetches the specified files for all takes. */
+	void FetchFiles(TUniquePtr<FBaseStream> InStream, TArray<FString> InFileNames);
+
+private:
+
+	inline static const float ConnectInterval = 5.0f;
+
+	void InitializeDelegates();
+
+	struct FEmpty
+	{
+	};
+
+	void ConnectControlClient(FEmpty);
+
+	void RegisterForAllEvents();
+	void OnCPSEvent(TSharedPtr<FControlUpdate> InUpdateMessage);
+
+	void StartConnectTimer(float InDelay = 0.0f);
+	void OnConnectTick();
+	void OnDisconnect(const FString& InCause);
+
+	static TSharedRef<FCaptureTimerManager> GetTimerManager();
+
+	TSharedRef<FCaptureTimerManager> TimerManager;
+
+	FString DeviceIpAddress;
+	uint16 DeviceControlPort;
+
+	std::atomic_bool bIsConnected;
+
+	using FConnectionThread = TQueueRunner<FEmpty>;
+	TUniquePtr<FConnectionThread> ConnThread;
+	FControlMessenger ControlMessenger;
+	FCaptureTimerManager::FTimerHandle ConnectTimerHandle;
+
+	TUniquePtr<FExportClient> ExportClient;
+
+	FCriticalSection Mutex;
+	TMap<FTakeId, FGetTakeMetadataResponse::FTakeObject> TakeMetadata;
+
+	TMap<FTakeId, FExportClient::FTaskId> IdMap;
+};
+
+}
\ No newline at end of file
diff --git a/Engine/Plugins/VirtualProduction/CaptureManager/CaptureManagerDevices/Source/CPSLiveLinkDevice/Public/Protocol/CPSFileStream.h b/Engine/Plugins/VirtualProduction/CaptureManager/CaptureManagerDevices/Source/CPSLiveLinkDevice/Public/Protocol/CPSFileStream.h
new file mode 100644
--- /dev/null
+++ b/Engine/Plugins/VirtualProduction/CaptureManager/CaptureManagerDevices/Source/CPSLiveLinkDevice/Public/Protocol/CPSFileStream.h
@@ -0,0 +1,45 @@
+// Copyright Epic Games, Inc. All Rights Reserved.
+
+#pragma once
+
+#include "ExportClient/ExportClient.h"
+
+namespace UE::CaptureManager
+{
+
+class CPSLIVELINKDEVICE_API FCPSFileStream : public UE::CaptureManager::FBaseStream
+{
+public:
+
+	DECLARE_DELEGATE_OneParam(FReportProgress, float InProgress);
+	DECLARE_DELEGATE_OneParam(FExportFinished, UE::CaptureManager::TProtocolResult<void> InResult);
+
+	FCPSFileStream(FString InBaseDir, uint64 InSize);
+
+	virtual bool StartFile(const FString& InTakeName, const FString& InFileName) override;
+
+	virtual bool ProcessData(const FString& InTakeName, const FString& InFileName, const TConstArrayView<uint8>& InData) override;
+	virtual bool FinishFile(const FString& InTakeName, const FString& InFileName, const TStaticArray<uint8, 16>& InHash) override;
+
+	virtual void Finalize(UE::CaptureManager::TProtocolResult<void> InResult) override;
+
+	void SetExportFinished(FExportFinished InExportFinished);
+	void SetProgressHandler(FReportProgress InReportProgress);
+
+private:
+
+	void ReportProgressStep(const FString& InFileName, const uint32 InArrivedSize);
+
+	FString BaseDir;
+	TUniquePtr<FArchive> Writer;
+	TUniquePtr<FMD5> MD5Generator;
+
+	uint64 TotalExportExpectedSize;
+	uint64 TotalExportArrivedSize;
+
+	FExportFinished OnExportFinished;
+	FReportProgress OnReportProgress;
+};
+
+}
+
