[CaptureManager] Fixed crash during thumbnail extraction

diff --git a/Engine/Plugins/VirtualProduction/CaptureManager/CaptureManagerDevices/Source/MonoVideoIngestDevice/Private/MonoVideoIngestDevice.cpp b/Engine/Plugins/VirtualProduction/CaptureManager/CaptureManagerDevices/Source/MonoVideoIngestDevice/Private/MonoVideoIngestDevice.cpp
--- a/Engine/Plugins/VirtualProduction/CaptureManager/CaptureManagerDevices/Source/MonoVideoIngestDevice/Private/MonoVideoIngestDevice.cpp
+++ b/Engine/Plugins/VirtualProduction/CaptureManager/CaptureManagerDevices/Source/MonoVideoIngestDevice/Private/MonoVideoIngestDevice.cpp
@@ -287,17 +287,20 @@ TOptional<FTakeMetadata> UMonoVideoIngestDevice::ReadTake(const FString& InCurre
 			FVideoDeviceThumbnailExtractor ThumbnailExtractor;
 			UE::CaptureManager::FMediaTextureSample* Sample = TextureSampleResult.GetValue().Get();
 
-			TArray<FColor> ThumbnailRawColorData = ThumbnailExtractor.ExtractThumbnail(Sample);
-
-			if (ThumbnailRawColorData.Num() > 0)
+			if (Sample)
 			{
-				FTakeThumbnailData::FRawImage RawImage;
-				RawImage.DecompressedImageData = MoveTemp(ThumbnailRawColorData);
-				RawImage.Width = Sample->Dimensions.X;
-				RawImage.Height = Sample->Dimensions.Y;
-				RawImage.Format = ERawImageFormat::BGRA8;
+				TArray<FColor> ThumbnailRawColorData = ThumbnailExtractor.ExtractThumbnail(Sample);
+
+				if (ThumbnailRawColorData.Num() > 0)
+				{
+					FTakeThumbnailData::FRawImage RawImage;
+					RawImage.DecompressedImageData = MoveTemp(ThumbnailRawColorData);
+					RawImage.Width = Sample->Dimensions.X;
+					RawImage.Height = Sample->Dimensions.Y;
+					RawImage.Format = ERawImageFormat::BGRA8;
 
-				TakeMetadata.Thumbnail = MoveTemp(RawImage);
+					TakeMetadata.Thumbnail = MoveTemp(RawImage);
+				}
 			}
 		}
 	}
diff --git a/Engine/Plugins/VirtualProduction/CaptureManager/CaptureManagerDevices/Source/StereoVideoIngestDevice/Private/StereoVideoIngestDevice.cpp b/Engine/Plugins/VirtualProduction/CaptureManager/CaptureManagerDevices/Source/StereoVideoIngestDevice/Private/StereoVideoIngestDevice.cpp
--- a/Engine/Plugins/VirtualProduction/CaptureManager/CaptureManagerDevices/Source/StereoVideoIngestDevice/Private/StereoVideoIngestDevice.cpp
+++ b/Engine/Plugins/VirtualProduction/CaptureManager/CaptureManagerDevices/Source/StereoVideoIngestDevice/Private/StereoVideoIngestDevice.cpp
@@ -513,17 +513,20 @@ TOptional<FTakeMetadata> UStereoVideoIngestDevice::CreateTakeMetadata(const FTak
 						FVideoDeviceThumbnailExtractor ThumbnailExtractor;
 						UE::CaptureManager::FMediaTextureSample* Sample = TextureSampleResult.GetValue().Get();
 
-						TArray<FColor> ThumbnailRawColorData = ThumbnailExtractor.ExtractThumbnail(Sample);
-
-						if (ThumbnailRawColorData.Num() > 0) 
+						if (Sample)
 						{
-							FTakeThumbnailData::FRawImage RawImage;
-							RawImage.DecompressedImageData = MoveTemp(ThumbnailRawColorData);
-							RawImage.Width = Sample->Dimensions.X;
-							RawImage.Height = Sample->Dimensions.Y;
-							RawImage.Format = ERawImageFormat::BGRA8;
+							TArray<FColor> ThumbnailRawColorData = ThumbnailExtractor.ExtractThumbnail(Sample);
+
+							if (ThumbnailRawColorData.Num() > 0)
+							{
+								FTakeThumbnailData::FRawImage RawImage;
+								RawImage.DecompressedImageData = MoveTemp(ThumbnailRawColorData);
+								RawImage.Width = Sample->Dimensions.X;
+								RawImage.Height = Sample->Dimensions.Y;
+								RawImage.Format = ERawImageFormat::BGRA8;
 
-							TakeMetadata.Thumbnail = MoveTemp(RawImage);
+								TakeMetadata.Thumbnail = MoveTemp(RawImage);
+							}
 						}
 					}
 				}
