[CaptureManager] Fixed crash when starting CPS file stream

diff --git a/Engine/Plugins/VirtualProduction/CaptureManager/CaptureManagerDevices/Source/CPSLiveLinkDevice/Private/Protocol/CPSFileStream.cpp b/Engine/Plugins/VirtualProduction/CaptureManager/CaptureManagerDevices/Source/CPSLiveLinkDevice/Private/Protocol/CPSFileStream.cpp
--- a/Engine/Plugins/VirtualProduction/CaptureManager/CaptureManagerDevices/Source/CPSLiveLinkDevice/Private/Protocol/CPSFileStream.cpp
+++ b/Engine/Plugins/VirtualProduction/CaptureManager/CaptureManagerDevices/Source/CPSLiveLinkDevice/Private/Protocol/CPSFileStream.cpp
@@ -22,6 +22,12 @@ bool FCPSFileStream::StartFile(const FString& InTakeName, const FString& InFileN
 	checkf(!Writer, TEXT("Writer must be nullptr when new file arrives"));
 	Writer = TUniquePtr<FArchive>(IFileManager::Get().CreateFileWriter(*FilePathToBeSaved, EFileWrite::FILEWRITE_None));
 
+	if (!Writer)
+	{
+		// If writer has not been assigned it is likely another writer is writing to FilePathToBeSaved
+		return false;
+	}
+
 	checkf(!MD5Generator, TEXT("MD5 generator must be nullptr when new file arrives"));
 	MD5Generator = MakeUnique<FMD5>();
 
