Invalid finalize at the hash error

diff --git a/Engine/Plugins/VirtualProduction/CaptureManager/CaptureManagerCore/Source/CaptureProtocolStack/Private/ExportClient/ExportClient.cpp b/Engine/Plugins/VirtualProduction/CaptureManager/CaptureManagerCore/Source/CaptureProtocolStack/Private/ExportClient/ExportClient.cpp
--- a/Engine/Plugins/VirtualProduction/CaptureManager/CaptureManagerCore/Source/CaptureProtocolStack/Private/ExportClient/ExportClient.cpp
+++ b/Engine/Plugins/VirtualProduction/CaptureManager/CaptureManagerCore/Source/CaptureProtocolStack/Private/ExportClient/ExportClient.cpp
@@ -245,10 +245,10 @@ TProtocolResult<void> FExportClient::ReceiveResponses(TMap<uint32, TPair<FString
 		const TPair<FString, FTakeFile>& TakeFile = InResponseMap[ResponseTransactionId];
 
 		const FTakeFile& File = TakeFile.Value;
-		if (!InBaseStream.StartFile(TakeFile.Key, File.FileName))
+		TProtocolResult<void> StartFileResult = InBaseStream.StartFile(TakeFile.Key, File.FileName);
+		if (StartFileResult.HasError())
 		{
-			// Error is already handled
-			return FCaptureProtocolError(TEXT("Stream error"));
+			return StartFileResult;
 		}
 
 		uint64 BytesLeft = File.Length;
@@ -273,10 +273,10 @@ TProtocolResult<void> FExportClient::ReceiveResponses(TMap<uint32, TPair<FString
 
 			TArray<uint8> FileData = FileDataResult.StealValue();
 
-			if (!InBaseStream.ProcessData(TakeFile.Key, File.FileName, FileData))
+			TProtocolResult<void> ProcessDataResult = InBaseStream.ProcessData(TakeFile.Key, File.FileName, FileData);
+			if (ProcessDataResult.HasError())
 			{
-				// Error is already handled
-				return FCaptureProtocolError(TEXT("Stream error"));
+				return ProcessDataResult;
 			}
 
 			BytesLeft -= FileData.Num();
@@ -288,9 +288,10 @@ TProtocolResult<void> FExportClient::ReceiveResponses(TMap<uint32, TPair<FString
 			return FileHashResult.StealError();
 		}
 
-		if (!InBaseStream.FinishFile(TakeFile.Key, File.FileName, FileHashResult.StealValue()))
+		TProtocolResult<void> FinishStream = InBaseStream.FinishFile(TakeFile.Key, File.FileName, FileHashResult.StealValue());
+		if (FinishStream.HasError())
 		{
-			return FCaptureProtocolError(TEXT("Stream error"));
+			return FinishStream;
 		}
 
 		InResponseMap.Remove(ResponseTransactionId);
diff --git a/Engine/Plugins/VirtualProduction/CaptureManager/CaptureManagerCore/Source/CaptureProtocolStack/Public/ExportClient/ExportWorker.h b/Engine/Plugins/VirtualProduction/CaptureManager/CaptureManagerCore/Source/CaptureProtocolStack/Public/ExportClient/ExportWorker.h
--- a/Engine/Plugins/VirtualProduction/CaptureManager/CaptureManagerCore/Source/CaptureProtocolStack/Public/ExportClient/ExportWorker.h
+++ b/Engine/Plugins/VirtualProduction/CaptureManager/CaptureManagerCore/Source/CaptureProtocolStack/Public/ExportClient/ExportWorker.h
@@ -103,9 +103,9 @@ public:
 	FBaseStream() = default;
 	virtual ~FBaseStream() = default;
 
-	virtual bool StartFile(const FString& InTakeName, const FString& InFileName) = 0;
-	virtual bool ProcessData(const FString& InTakeName, const FString& InFileName, const TConstArrayView<uint8>& InData) = 0;
-	virtual bool FinishFile(const FString& InTakeName, const FString& InFileName, const TStaticArray<uint8, 16>& InHash) = 0;
+	virtual TProtocolResult<void> StartFile(const FString& InTakeName, const FString& InFileName) = 0;
+	virtual TProtocolResult<void> ProcessData(const FString& InTakeName, const FString& InFileName, const TConstArrayView<uint8>& InData) = 0;
+	virtual TProtocolResult<void> FinishFile(const FString& InTakeName, const FString& InFileName, const TStaticArray<uint8, 16>& InHash) = 0;
 
 	virtual void Finalize(TProtocolResult<void> InResult) = 0;
 };
diff --git a/Engine/Plugins/VirtualProduction/CaptureManager/CaptureManagerDevices/Source/CPSLiveLinkDevice/Private/Protocol/CPSDataStream.cpp b/Engine/Plugins/VirtualProduction/CaptureManager/CaptureManagerDevices/Source/CPSLiveLinkDevice/Private/Protocol/CPSDataStream.cpp
--- a/Engine/Plugins/VirtualProduction/CaptureManager/CaptureManagerDevices/Source/CPSLiveLinkDevice/Private/Protocol/CPSDataStream.cpp
+++ b/Engine/Plugins/VirtualProduction/CaptureManager/CaptureManagerDevices/Source/CPSLiveLinkDevice/Private/Protocol/CPSDataStream.cpp
@@ -12,19 +12,19 @@ FCPSDataStream::FCPSDataStream(FFileExportFinished InFileExportFinished)
 {
 }
 
-bool FCPSDataStream::StartFile(const FString& InTakeName, const FString& InFileName)
+TProtocolResult<void> FCPSDataStream::StartFile(const FString& InTakeName, const FString& InFileName)
 {
 	checkf(Data.IsEmpty(), TEXT("Data buffer must be empty"));
-	return true;
+	return ResultOk;
 }
 
-bool FCPSDataStream::ProcessData(const FString& InTakeName, const FString& InFileName, const TConstArrayView<uint8>& InData)
+TProtocolResult<void> FCPSDataStream::ProcessData(const FString& InTakeName, const FString& InFileName, const TConstArrayView<uint8>& InData)
 {
 	Data.Append(InData);
-	return true;
+	return ResultOk;
 }
 
-bool FCPSDataStream::FinishFile(const FString& InTakeName, const FString& InFileName, const TStaticArray<uint8, 16>& InHash)
+TProtocolResult<void> FCPSDataStream::FinishFile(const FString& InTakeName, const FString& InFileName, const TStaticArray<uint8, 16>& InHash)
 {
 	using namespace UE::CaptureManager;
 
@@ -38,13 +38,15 @@ bool FCPSDataStream::FinishFile(const FString& InTakeName, const FString& InFile
 	{
 		TMap<FString, UE::CaptureManager::TProtocolResult<FData>>& TakeEntry = ExportResults.FindOrAdd(InTakeName);
 		TakeEntry.Add(InFileName, FCaptureProtocolError(TEXT("Invalid file hash")));
-		return true;
+
+		// Continue sending the data regardless of the error
+		return ResultOk;
 	}
 
 	TMap<FString, UE::CaptureManager::TProtocolResult<FData>>& TakeEntry = ExportResults.FindOrAdd(InTakeName);
 	TakeEntry.Add(InFileName, MoveTemp(Data));
 
-	return true;
+	return ResultOk;
 }
 
 void FCPSDataStream::Finalize(UE::CaptureManager::TProtocolResult<void> InResult)
diff --git a/Engine/Plugins/VirtualProduction/CaptureManager/CaptureManagerDevices/Source/CPSLiveLinkDevice/Private/Protocol/CPSFileStream.cpp b/Engine/Plugins/VirtualProduction/CaptureManager/CaptureManagerDevices/Source/CPSLiveLinkDevice/Private/Protocol/CPSFileStream.cpp
--- a/Engine/Plugins/VirtualProduction/CaptureManager/CaptureManagerDevices/Source/CPSLiveLinkDevice/Private/Protocol/CPSFileStream.cpp
+++ b/Engine/Plugins/VirtualProduction/CaptureManager/CaptureManagerDevices/Source/CPSLiveLinkDevice/Private/Protocol/CPSFileStream.cpp
@@ -15,7 +15,7 @@ FCPSFileStream::FCPSFileStream(FString InBaseDir, uint64 InSize)
 {
 }
 
-bool FCPSFileStream::StartFile(const FString& InTakeName, const FString& InFileName)
+TProtocolResult<void> FCPSFileStream::StartFile(const FString& InTakeName, const FString& InFileName)
 {
 	FString FilePathToBeSaved = FPaths::Combine(BaseDir, InTakeName, InFileName);
 
@@ -25,26 +25,26 @@ bool FCPSFileStream::StartFile(const FString& InTakeName, const FString& InFileN
 	if (!Writer)
 	{
 		// If writer has not been assigned it is likely another writer is writing to FilePathToBeSaved
-		return false;
+		return FCaptureProtocolError(TEXT("Unable to create a writer for file"));
 	}
 
 	checkf(!MD5Generator, TEXT("MD5 generator must be nullptr when new file arrives"));
 	MD5Generator = MakeUnique<FMD5>();
 
-	return true;
+	return ResultOk;
 }
 
-bool FCPSFileStream::ProcessData(const FString& InTakeName, const FString& InFileName, const TConstArrayView<uint8>& InData)
+TProtocolResult<void> FCPSFileStream::ProcessData(const FString& InTakeName, const FString& InFileName, const TConstArrayView<uint8>& InData)
 {
 	Writer->Serialize(const_cast<uint8*>(InData.GetData()), InData.Num());
 	ReportProgressStep(InFileName, InData.Num());
 
 	MD5Generator->Update(InData.GetData(), InData.Num());
 
-	return true;
+	return ResultOk;
 }
 
-bool FCPSFileStream::FinishFile(const FString& InTakeName, const FString& InFileName, const TStaticArray<uint8, 16>& InHash)
+TProtocolResult<void> FCPSFileStream::FinishFile(const FString& InTakeName, const FString& InFileName, const TStaticArray<uint8, 16>& InHash)
 {
 	checkf(Writer, TEXT("Writer must be valid when file has arrived"));
 	Writer->Close();
@@ -60,12 +60,10 @@ bool FCPSFileStream::FinishFile(const FString& InTakeName, const FString& InFile
 
 	if (FMemory::Memcmp(Hash.GetData(), InHash.GetData(), Hash.Num()) != 0)
 	{
-		Finalize(UE::CaptureManager::FCaptureProtocolError(TEXT("Invalid hash file")));
-
-		return false;
+		return FCaptureProtocolError(TEXT("Invalid hash file"));
 	}
 
-	return true;
+	return ResultOk;
 }
 
 void FCPSFileStream::Finalize(UE::CaptureManager::TProtocolResult<void> InResult)
diff --git a/Engine/Plugins/VirtualProduction/CaptureManager/CaptureManagerDevices/Source/CPSLiveLinkDevice/Public/Protocol/CPSDataStream.h b/Engine/Plugins/VirtualProduction/CaptureManager/CaptureManagerDevices/Source/CPSLiveLinkDevice/Public/Protocol/CPSDataStream.h
--- a/Engine/Plugins/VirtualProduction/CaptureManager/CaptureManagerDevices/Source/CPSLiveLinkDevice/Public/Protocol/CPSDataStream.h
+++ b/Engine/Plugins/VirtualProduction/CaptureManager/CaptureManagerDevices/Source/CPSLiveLinkDevice/Public/Protocol/CPSDataStream.h
@@ -18,9 +18,9 @@ public:
 
 	FCPSDataStream(FFileExportFinished InFileExportFinished);
 
-	bool StartFile(const FString& InTakeName, const FString& InFileName) override;
-	bool ProcessData(const FString& InTakeName, const FString& InFileName, const TConstArrayView<uint8>& InData) override;
-	bool FinishFile(const FString& InTakeName, const FString& InFileName, const TStaticArray<uint8, 16>& InHash) override;
+	TProtocolResult<void> StartFile(const FString& InTakeName, const FString& InFileName) override;
+	TProtocolResult<void> ProcessData(const FString& InTakeName, const FString& InFileName, const TConstArrayView<uint8>& InData) override;
+	TProtocolResult<void> FinishFile(const FString& InTakeName, const FString& InFileName, const TStaticArray<uint8, 16>& InHash) override;
 
 	void Finalize(UE::CaptureManager::TProtocolResult<void> InResult) override;
 
diff --git a/Engine/Plugins/VirtualProduction/CaptureManager/CaptureManagerDevices/Source/CPSLiveLinkDevice/Public/Protocol/CPSFileStream.h b/Engine/Plugins/VirtualProduction/CaptureManager/CaptureManagerDevices/Source/CPSLiveLinkDevice/Public/Protocol/CPSFileStream.h
--- a/Engine/Plugins/VirtualProduction/CaptureManager/CaptureManagerDevices/Source/CPSLiveLinkDevice/Public/Protocol/CPSFileStream.h
+++ b/Engine/Plugins/VirtualProduction/CaptureManager/CaptureManagerDevices/Source/CPSLiveLinkDevice/Public/Protocol/CPSFileStream.h
@@ -16,10 +16,10 @@ public:
 
 	FCPSFileStream(FString InBaseDir, uint64 InSize);
 
-	virtual bool StartFile(const FString& InTakeName, const FString& InFileName) override;
+	virtual TProtocolResult<void> StartFile(const FString& InTakeName, const FString& InFileName) override;
 
-	virtual bool ProcessData(const FString& InTakeName, const FString& InFileName, const TConstArrayView<uint8>& InData) override;
-	virtual bool FinishFile(const FString& InTakeName, const FString& InFileName, const TStaticArray<uint8, 16>& InHash) override;
+	virtual TProtocolResult<void> ProcessData(const FString& InTakeName, const FString& InFileName, const TConstArrayView<uint8>& InData) override;
+	virtual TProtocolResult<void> FinishFile(const FString& InTakeName, const FString& InFileName, const TStaticArray<uint8, 16>& InHash) override;
 
 	virtual void Finalize(UE::CaptureManager::TProtocolResult<void> InResult) override;
 
