Thumbnails not showing up for technoprops video

diff --git a/Engine/Plugins/VirtualProduction/CaptureManager/CaptureManagerDevices/Source/MonoVideoIngestDevice/Private/MonoVideoIngestDevice.cpp b/Engine/Plugins/VirtualProduction/CaptureManager/CaptureManagerDevices/Source/MonoVideoIngestDevice/Private/MonoVideoIngestDevice.cpp
--- a/Engine/Plugins/VirtualProduction/CaptureManager/CaptureManagerDevices/Source/MonoVideoIngestDevice/Private/MonoVideoIngestDevice.cpp
+++ b/Engine/Plugins/VirtualProduction/CaptureManager/CaptureManagerDevices/Source/MonoVideoIngestDevice/Private/MonoVideoIngestDevice.cpp
@@ -281,34 +281,20 @@ TOptional<FTakeMetadata> UMonoVideoIngestDevice::ReadTake(const FString& InCurre
 	{
 		TUniquePtr<IVideoReader> VideoReader = VideoReaderResult.StealValue();
 		Video.FrameRate = static_cast<float>(VideoReader->GetFrameRate().AsDecimal());
-		TValueOrError<TUniquePtr<UE::CaptureManager::FMediaTextureSample>, FText> TextureSampleResult = VideoReader->Next();
-		if (TextureSampleResult.HasValue())
-		{
-			FVideoDeviceThumbnailExtractor ThumbnailExtractor;
-			UE::CaptureManager::FMediaTextureSample* Sample = TextureSampleResult.GetValue().Get();
-
-			if (Sample)
-			{
-				TArray<FColor> ThumbnailRawColorData = ThumbnailExtractor.ExtractThumbnail(Sample);
-
-				if (ThumbnailRawColorData.Num() > 0)
-				{
-					FTakeThumbnailData::FRawImage RawImage;
-					RawImage.DecompressedImageData = MoveTemp(ThumbnailRawColorData);
-					RawImage.Width = Sample->Dimensions.X;
-					RawImage.Height = Sample->Dimensions.Y;
-					RawImage.Format = ERawImageFormat::BGRA8;
-
-					TakeMetadata.Thumbnail = MoveTemp(RawImage);
-				}
-			}
-		}
 	}
 	else
 	{
 		UE_LOG(LogVideoLiveLinkDevice, Warning, TEXT("Failed to determine the frame rate for the video file %s. Consider enabling Third Party Encoder in Capture Manager settings."), *InCurrentTakeFile);
 	}
 
+	FVideoDeviceThumbnailExtractor ThumbnailExtractor;
+	TOptional<FTakeThumbnailData::FRawImage> RawImageOpt = ThumbnailExtractor.ExtractThumbnail(InCurrentTakeFile);
+
+	if (RawImageOpt.IsSet())
+	{
+		TakeMetadata.Thumbnail = MoveTemp(RawImageOpt.GetValue());
+	}
+
 	Video.Name = MoveTemp(Name);
 	Video.Format = FPaths::GetExtension(InCurrentTakeFile);
 	Video.Path = InCurrentTakeFile;
diff --git a/Engine/Plugins/VirtualProduction/CaptureManager/CaptureManagerDevices/Source/StereoVideoIngestDevice/Private/StereoVideoIngestDevice.cpp b/Engine/Plugins/VirtualProduction/CaptureManager/CaptureManagerDevices/Source/StereoVideoIngestDevice/Private/StereoVideoIngestDevice.cpp
--- a/Engine/Plugins/VirtualProduction/CaptureManager/CaptureManagerDevices/Source/StereoVideoIngestDevice/Private/StereoVideoIngestDevice.cpp
+++ b/Engine/Plugins/VirtualProduction/CaptureManager/CaptureManagerDevices/Source/StereoVideoIngestDevice/Private/StereoVideoIngestDevice.cpp
@@ -506,35 +506,23 @@ TOptional<FTakeMetadata> UStereoVideoIngestDevice::CreateTakeMetadata(const FTak
 					{
 						VideoFrameRate = UE::CaptureManager::ParseFrameRate(Video.FrameRate);
 					}
-
-					TValueOrError<TUniquePtr<UE::CaptureManager::FMediaTextureSample>, FText> TextureSampleResult = VideoReader->Next();
-					if (TextureSampleResult.HasValue())
-					{
-						FVideoDeviceThumbnailExtractor ThumbnailExtractor;
-						UE::CaptureManager::FMediaTextureSample* Sample = TextureSampleResult.GetValue().Get();
-
-						if (Sample)
-						{
-							TArray<FColor> ThumbnailRawColorData = ThumbnailExtractor.ExtractThumbnail(Sample);
-
-							if (ThumbnailRawColorData.Num() > 0)
-							{
-								FTakeThumbnailData::FRawImage RawImage;
-								RawImage.DecompressedImageData = MoveTemp(ThumbnailRawColorData);
-								RawImage.Width = Sample->Dimensions.X;
-								RawImage.Height = Sample->Dimensions.Y;
-								RawImage.Format = ERawImageFormat::BGRA8;
-
-								TakeMetadata.Thumbnail = MoveTemp(RawImage);
-							}
-						}
-					}
 				}
 				else
 				{
 					UE_LOG(LogVideoLiveLinkDevice, Warning, TEXT("Failed to determine the frame rate for the video file %s. Consider enabling Third Party Encoder in Capture Manager settings."), *TakeComponent.Path);
 				}
 
+				if (!TakeMetadata.Thumbnail.GetRawImage().IsSet())
+				{
+					FVideoDeviceThumbnailExtractor ThumbnailExtractor;
+					TOptional<FTakeThumbnailData::FRawImage> RawImageOpt = ThumbnailExtractor.ExtractThumbnail(TakeComponent.Path);
+
+					if (RawImageOpt.IsSet())
+					{
+						TakeMetadata.Thumbnail = MoveTemp(RawImageOpt.GetValue());
+					}
+				}
+
 				Video.Path = TakeComponent.Path;
 				
 				FString MediaFilePath = FPaths::ConvertRelativePathToFull(InTakeWithComponents.TakePath, Video.Path);
diff --git a/Engine/Plugins/VirtualProduction/CaptureManager/CaptureManagerDevices/Source/VideoLiveLinkDeviceCommon/Private/Utils/VideoDeviceThumbnailExtractor.cpp b/Engine/Plugins/VirtualProduction/CaptureManager/CaptureManagerDevices/Source/VideoLiveLinkDeviceCommon/Private/Utils/VideoDeviceThumbnailExtractor.cpp
--- a/Engine/Plugins/VirtualProduction/CaptureManager/CaptureManagerDevices/Source/VideoLiveLinkDeviceCommon/Private/Utils/VideoDeviceThumbnailExtractor.cpp
+++ b/Engine/Plugins/VirtualProduction/CaptureManager/CaptureManagerDevices/Source/VideoLiveLinkDeviceCommon/Private/Utils/VideoDeviceThumbnailExtractor.cpp
@@ -7,89 +7,332 @@
 #include "CaptureManagerMediaRWModule.h"
 #include "VideoLiveLinkDeviceLog.h"
 
+#include "Settings/CaptureManagerSettings.h"
+
 #include "Utils/MediaPixelFormatConversions.h"
 
 DEFINE_LOG_CATEGORY(LogVideoLiveLinkDevice);
 
-
 namespace UE::CaptureManager
 {
 
+namespace Private
+{
 
-FVideoDeviceThumbnailExtractor::FVideoDeviceThumbnailExtractor()
+static bool IsEncoderFFmpeg(const FString& InEncoderPath)
 {
+	// Get base filename without extension
+	const FString EncoderFileName = FPaths::GetBaseFilename(InEncoderPath);
+	return EncoderFileName == TEXT("ffmpeg");
 }
 
-TArray<FColor> FVideoDeviceThumbnailExtractor::ExtractThumbnail(UE::CaptureManager::FMediaTextureSample* InSample)
+static FString GetFFprobePath(const FString& InEncoderPath)
 {
-	EMediaTexturePixelFormat SampleFormat = InSample->CurrentFormat;
-	TArray<uint8>& Buffer = InSample->Buffer;
+	check(IsEncoderFFmpeg(InEncoderPath));
 
-	TArray<FColor> ThumbnailRawColorData;
-	switch (SampleFormat)
-	{
-	case EMediaTexturePixelFormat::U8_Mono:
+	// Get base filename with extension: e.g ffmpeg.exe
+	const FString EncoderFileName = FPaths::GetCleanFilename(InEncoderPath);
+	const FString RootPath = FPaths::GetPath(InEncoderPath);
+	const FString FFprobeFileName = EncoderFileName.Replace(TEXT("ffmpeg"), TEXT("ffprobe"));
+	const FString FFprobePath = RootPath / FFprobeFileName;
+	return FFprobePath;
+}
+
+}
+
+FVideoDeviceThumbnailExtractor::FVideoDeviceThumbnailExtractor() = default;
+
+TOptional<FTakeThumbnailData::FRawImage> FVideoDeviceThumbnailExtractor::ExtractThumbnail(const FString& InCurrentFile)
+{
+	FMediaRWManager& MediaManager = FModuleManager::LoadModuleChecked<FCaptureManagerMediaRWModule>("CaptureManagerMediaRW").Get();
+
+	TArray<FColor> Thumbnail;
+	int32 Width = 0;
+	int32 Height = 0;
+
+	TValueOrError<TUniquePtr<IVideoReader>, FText> VideoReaderResult = MediaManager.CreateVideoReader(InCurrentFile);
+	if (VideoReaderResult.HasValue())
 	{
-		ThumbnailRawColorData.Reserve(Buffer.Num());
-		for (uint8& Value : Buffer)
+		TUniquePtr<IVideoReader> VideoReader = VideoReaderResult.StealValue();
+		TValueOrError<TUniquePtr<UE::CaptureManager::FMediaTextureSample>, FText> TextureSampleResult = VideoReader->Next();
+		if (TextureSampleResult.HasValue())
 		{
-			ThumbnailRawColorData.Add(FColor(Value, Value, Value));
+			const UE::CaptureManager::FMediaTextureSample* Sample = TextureSampleResult.GetValue().Get();
+
+			if (Sample)
+			{
+				Thumbnail = ConvertThumbnailFromSample(Sample);
+				Width = Sample->Dimensions.X;
+				Height = Sample->Dimensions.Y;
+			}
 		}
-		break;
+		else
+		{
+			FText ErrorText = TextureSampleResult.StealError();
+			UE_LOG(LogVideoLiveLinkDevice, Warning, TEXT("Couldn't obtain the thumbnail from the video %s : %s"), *InCurrentFile, *ErrorText.ToString());
+		}
+	}
+	else
+	{
+		FText ErrorText = VideoReaderResult.StealError();
+		UE_LOG(LogVideoLiveLinkDevice, Warning, TEXT("Couldn't open video file %s : %s"), *InCurrentFile, *ErrorText.ToString());
 	}
-	case EMediaTexturePixelFormat::U8_RGB:
+	
+	if (Thumbnail.IsEmpty())
 	{
-		ThumbnailRawColorData.Reserve(Buffer.Num() / 3);
-		for (int32 ValueIndex = 0; ValueIndex < Buffer.Num(); ValueIndex += 3)
+		const UCaptureManagerSettings* Settings = GetDefault<UCaptureManagerSettings>();
+		if (Settings->bEnableThirdPartyEncoder)
 		{
-			ThumbnailRawColorData.Add(FColor(Buffer[ValueIndex], Buffer[ValueIndex + 1], Buffer[ValueIndex + 2]));
+			UE_LOG(LogVideoLiveLinkDevice, Warning, TEXT("Obtaining thumbnail using third party encoder"));
+
+			Thumbnail = ObtainThumbnailFromThirdPartyEncoder(Settings->ThirdPartyEncoder.FilePath, InCurrentFile);
+			TOptional<FIntPoint> DimensionsOpt = ObtainImageDimensionsFromThirdPartyEncoder(Settings->ThirdPartyEncoder.FilePath, InCurrentFile);
+
+			if (DimensionsOpt.IsSet())
+			{
+				FIntPoint Dimensions = MoveTemp(DimensionsOpt.GetValue());
+
+				Width = ThirdPartyEncoderThumbnailWidth;
+				Height = FMath::RoundToInt((static_cast<float>(Width) * Dimensions.Y) / Dimensions.X);
+
+				check(Thumbnail.Num() == Width * Height);
+			}
+			else
+			{
+				UE_LOG(LogVideoLiveLinkDevice, Warning, TEXT("Couldn't obtain frame dimensions of the video file %s"), *InCurrentFile);
+				return {};
+			}
 		}
-		break;
 	}
-	case EMediaTexturePixelFormat::U8_BGR:
+
+	if (Thumbnail.IsEmpty())
+	{
+		return {};
+	}
+
+	FTakeThumbnailData::FRawImage RawImage;
+	RawImage.DecompressedImageData = MoveTemp(Thumbnail);
+	RawImage.Width = Width;
+	RawImage.Height = Height;
+	RawImage.Format = ERawImageFormat::BGRA8;
+
+	return RawImage;
+}
+
+TArray<FColor> FVideoDeviceThumbnailExtractor::ConvertThumbnailFromSample(const FMediaTextureSample* InSample)
+{
+	EMediaTexturePixelFormat SampleFormat = InSample->CurrentFormat;
+	const TArray<uint8>& Buffer = InSample->Buffer;
+
+	TArray<FColor> ThumbnailRawColorData;
+	switch (SampleFormat)
 	{
-		ThumbnailRawColorData.Reserve(Buffer.Num() / 3);
-		for (int32 ValueIndex = 0; ValueIndex < Buffer.Num(); ValueIndex += 3)
+		case EMediaTexturePixelFormat::U8_Mono:
+		{
+			ThumbnailRawColorData.Reserve(Buffer.Num());
+			for (const uint8& Value : Buffer)
+			{
+				ThumbnailRawColorData.Add(FColor(Value, Value, Value));
+			}
+			break;
+		}
+		case EMediaTexturePixelFormat::U8_RGB:
+		{
+			const int32 NumChannels = 3;
+			check(Buffer.Num() % NumChannels == 0);
+
+			ThumbnailRawColorData.Reserve(Buffer.Num() / NumChannels);
+			for (int32 ValueIndex = 0; ValueIndex < Buffer.Num(); ValueIndex += NumChannels)
+			{
+				ThumbnailRawColorData.Add(FColor(Buffer[ValueIndex], Buffer[ValueIndex + 1], Buffer[ValueIndex + 2]));
+			}
+			break;
+		}
+		case EMediaTexturePixelFormat::U8_BGR:
 		{
-			ThumbnailRawColorData.Add(FColor(Buffer[ValueIndex + 2], Buffer[ValueIndex + 1], Buffer[ValueIndex]));
+			const int32 NumChannels = 3;
+			check(Buffer.Num() % NumChannels == 0);
+
+			ThumbnailRawColorData.Reserve(Buffer.Num() / NumChannels);
+			for (int32 ValueIndex = 0; ValueIndex < Buffer.Num(); ValueIndex += NumChannels)
+			{
+				ThumbnailRawColorData.Add(FColor(Buffer[ValueIndex + 2], Buffer[ValueIndex + 1], Buffer[ValueIndex]));
+			}
+			break;
+		}
+		case EMediaTexturePixelFormat::U8_RGBA:
+		{
+			const int32 NumChannels = 4;
+			check(Buffer.Num() % NumChannels == 0);
+
+			ThumbnailRawColorData.Reserve(Buffer.Num() / NumChannels);
+			for (int32 ValueIndex = 0; ValueIndex < Buffer.Num(); ValueIndex += NumChannels)
+			{
+				ThumbnailRawColorData.Add(FColor(Buffer[ValueIndex], Buffer[ValueIndex + 1], Buffer[ValueIndex + 2], Buffer[ValueIndex + 3]));
+			}
+			break;
+		}
+		case EMediaTexturePixelFormat::U8_BGRA:
+		{
+			const int32 NumChannels = 4;
+			check(Buffer.Num() % NumChannels == 0);
+
+			ThumbnailRawColorData.Reserve(Buffer.Num() / NumChannels);
+			for (int32 ValueIndex = 0; ValueIndex < Buffer.Num(); ValueIndex += NumChannels)
+			{
+				ThumbnailRawColorData.Add(FColor(Buffer[ValueIndex + 2], Buffer[ValueIndex + 1], Buffer[ValueIndex], Buffer[ValueIndex + 3]));
+			}
+			break;
+		}
+		case EMediaTexturePixelFormat::U8_I420:
+			ThumbnailRawColorData = UE::CaptureManager::UEConvertI420ToBGRA(InSample);
+			break;
+		case EMediaTexturePixelFormat::U8_NV12:
+			ThumbnailRawColorData = UE::CaptureManager::UEConvertNV12ToBGRA(InSample);
+			break;
+		case EMediaTexturePixelFormat::U8_YUY2:
+			ThumbnailRawColorData = UE::CaptureManager::UEConvertYUY2ToBGRA(InSample);
+			break;
+		default:
+		{
+			UE_LOG(LogVideoLiveLinkDevice, Warning, TEXT("Unsupported image format"));
 		}
-		break;
 	}
-	case EMediaTexturePixelFormat::U8_RGBA:
+
+	return ThumbnailRawColorData;
+}
+
+TArray<FColor> FVideoDeviceThumbnailExtractor::ObtainThumbnailFromThirdPartyEncoder(const FString& InEncoderPath, const FString& InCurrentFile)
+{
+	FString ScaleArg = FString::Printf(TEXT("scale=%d:-1"), ThirdPartyEncoderThumbnailWidth);
+
+	FString CommandLineArgs = FString::Format(TEXT("-hide_banner -loglevel error -i \"{0}\" -map v:0 -vf \"thumbnail=2, {1}\" -frames:v 1 -f rawvideo -pix_fmt rgb24 -"), { InCurrentFile, ScaleArg });
+	
+	TArray<uint8> ThumbnailRawData = RunProcess(InEncoderPath, CommandLineArgs);
+
+	TArray<FColor> ThumbnailRawColorData;
+	ThumbnailRawColorData.Reserve(ThumbnailRawData.Num() / 3);
+
+	for (int32 ValueIndex = 0; ValueIndex < ThumbnailRawData.Num(); ValueIndex += 3)
+	{
+		ThumbnailRawColorData.Add(FColor(ThumbnailRawData[ValueIndex], ThumbnailRawData[ValueIndex + 1], ThumbnailRawData[ValueIndex + 2]));
+	}
+
+	return ThumbnailRawColorData;
+}
+
+TOptional<FIntPoint> FVideoDeviceThumbnailExtractor::ObtainImageDimensionsFromThirdPartyEncoder(const FString& InEncoderPath, const FString& InCurrentFile)
+{
+	if (!Private::IsEncoderFFmpeg(InEncoderPath))
 	{
-		ThumbnailRawColorData.Reserve(Buffer.Num() / 4);
-		for (int32 ValueIndex = 0; ValueIndex < Buffer.Num(); ValueIndex += 4)
+		return {};
+	}
+
+	FString CommandLine = FString::Format(TEXT("-v error -select_streams v:0 -show_entries stream=width,height -of default=noprint_wrappers=1 {0}"), { InCurrentFile });
+	FString EncoderPath = Private::GetFFprobePath(InEncoderPath);
+
+	TArray<uint8> WidthAndHeight = RunProcess(EncoderPath, CommandLine);
+
+	if (WidthAndHeight.IsEmpty())
+	{
+		return {};
+	}
+
+	FString WidthAndHeightString = 
+		FString::ConstructFromPtrSize(reinterpret_cast<const UTF8CHAR*>(WidthAndHeight.GetData()), WidthAndHeight.Num());
+
+	TArray<FString> Lines;
+	WidthAndHeightString.ParseIntoArrayLines(Lines);
+
+	FIntPoint Dimensions = FIntPoint::ZeroValue;
+	for (const FString& Line : Lines)
+	{
+		TArray<FString> KeyValue;
+		Line.ParseIntoArray(KeyValue, TEXT("="));
+
+		if (KeyValue.Num() >= 2)
 		{
-			ThumbnailRawColorData.Add(FColor(Buffer[ValueIndex], Buffer[ValueIndex + 1], Buffer[ValueIndex + 2], Buffer[ValueIndex + 3]));
+			const FString& Key = KeyValue[0];
+			const FString& Value = KeyValue[1];
+
+			if (Key.Equals(TEXT("width"), ESearchCase::IgnoreCase))
+			{
+				Dimensions.X = FCString::Atoi(*Value);
+			}
+			else if (Key.Equals(TEXT("height"), ESearchCase::IgnoreCase))
+			{
+				Dimensions.Y = FCString::Atoi(*Value);
+			}
 		}
-		break;
 	}
-	case EMediaTexturePixelFormat::U8_BGRA:
+
+	if (Dimensions == FIntPoint::ZeroValue)
 	{
-		ThumbnailRawColorData.Reserve(Buffer.Num() / 4);
-		for (int32 ValueIndex = 0; ValueIndex < Buffer.Num(); ValueIndex += 4)
+		return {};
+	}
+
+	return Dimensions;
+}
+
+TArray<uint8> FVideoDeviceThumbnailExtractor::RunProcess(const FString& InProcessName, const FString& InProcessArgs)
+{
+	void* ReadPipe = nullptr;
+	void* WritePipe = nullptr;
+	verify(FPlatformProcess::CreatePipe(ReadPipe, WritePipe, false));
+
+	constexpr bool bLaunchDetached = false;
+	constexpr bool bLaunchHidden = true;
+	constexpr bool bLaunchReallyHidden = true;
+	FProcHandle ProcHandle =
+		FPlatformProcess::CreateProc(*InProcessName, *InProcessArgs, bLaunchDetached, bLaunchHidden, bLaunchReallyHidden, nullptr, 0, nullptr, WritePipe, nullptr);
+
+	ON_SCOPE_EXIT
+	{
+		FPlatformProcess::TerminateProc(ProcHandle);
+		FPlatformProcess::CloseProc(ProcHandle);
+		FPlatformProcess::ClosePipe(ReadPipe, WritePipe);
+	};
+
+	if (!ProcHandle.IsValid())
+	{
+		return TArray<uint8>();
+	}
+
+	TArray<uint8> FullCommandOutput;
+	while (FPlatformProcess::IsProcRunning(ProcHandle))
+	{
+		TArray<uint8> CommandOutput = ReadFromPipe(ReadPipe);
+
+		if (CommandOutput.IsEmpty())
 		{
-			ThumbnailRawColorData.Add(FColor(Buffer[ValueIndex + 2], Buffer[ValueIndex + 1], Buffer[ValueIndex], Buffer[ValueIndex + 3]));
+			FPlatformProcess::Sleep(0.010f);
 		}
-		break;
+
+		FullCommandOutput.Append(MoveTemp(CommandOutput));
 	}
-	case EMediaTexturePixelFormat::U8_I420:
-		ThumbnailRawColorData = UE::CaptureManager::UEConvertI420ToBGRA(InSample);
-		break;
-	case EMediaTexturePixelFormat::U8_NV12:
-		ThumbnailRawColorData = UE::CaptureManager::UEConvertNV12ToBGRA(InSample);
-		break;
-	case EMediaTexturePixelFormat::U8_YUY2:
-		ThumbnailRawColorData = UE::CaptureManager::UEConvertYUY2ToBGRA(InSample);
-		break;
-	default:
+
+	int32 ReturnCode = 0;
+	FPlatformProcess::GetProcReturnCode(ProcHandle, &ReturnCode);
+
+	if (ReturnCode != 0)
 	{
-		UE_LOG(LogVideoLiveLinkDevice, Warning, TEXT("Unsupported image format"));
+		return TArray<uint8>();
 	}
+
+	FullCommandOutput.Append(ReadFromPipe(ReadPipe));
+	return FullCommandOutput;
+}
+
+TArray<uint8> FVideoDeviceThumbnailExtractor::ReadFromPipe(void* InPipe)
+{
+	TArray<uint8> CommandOutput;
+	bool Read = FPlatformProcess::ReadPipeToArray(InPipe, CommandOutput);
+	if (!Read)
+	{
+		CommandOutput.Empty();
 	}
 
-	InSample->Buffer.Empty();
-	return ThumbnailRawColorData;
+	return CommandOutput;
 }
 
-}
\ No newline at end of file
+}
diff --git a/Engine/Plugins/VirtualProduction/CaptureManager/CaptureManagerDevices/Source/VideoLiveLinkDeviceCommon/Public/Utils/VideoDeviceThumbnailExtractor.h b/Engine/Plugins/VirtualProduction/CaptureManager/CaptureManagerDevices/Source/VideoLiveLinkDeviceCommon/Public/Utils/VideoDeviceThumbnailExtractor.h
--- a/Engine/Plugins/VirtualProduction/CaptureManager/CaptureManagerDevices/Source/VideoLiveLinkDeviceCommon/Public/Utils/VideoDeviceThumbnailExtractor.h
+++ b/Engine/Plugins/VirtualProduction/CaptureManager/CaptureManagerDevices/Source/VideoLiveLinkDeviceCommon/Public/Utils/VideoDeviceThumbnailExtractor.h
@@ -1,18 +1,39 @@
 // Copyright Epic Games, Inc. All Rights Reserved.
 
 #pragma once
+
 #include "Math/Color.h"
+#include "CaptureManagerTakeMetadata.h"
+
+#include "Misc/CoreMiscDefines.h"
 
+#define UE_API VIDEOLIVELINKDEVICECOMMON_API
 
 namespace UE::CaptureManager 
 {
 struct FMediaTextureSample;
 
-class VIDEOLIVELINKDEVICECOMMON_API FVideoDeviceThumbnailExtractor
+class FVideoDeviceThumbnailExtractor
 {
 public:
-	FVideoDeviceThumbnailExtractor();
-	TArray<FColor> ExtractThumbnail(FMediaTextureSample* Sample);
+
+	UE_INTERNAL UE_API FVideoDeviceThumbnailExtractor();
+
+	UE_INTERNAL UE_API TOptional<FTakeThumbnailData::FRawImage> ExtractThumbnail(const FString& InCurrentFile);
+
+private:
+
+	TArray<FColor> ConvertThumbnailFromSample(const struct FMediaTextureSample* InSample);
+
+	static constexpr int32 ThirdPartyEncoderThumbnailWidth = 256;
+
+	TArray<FColor> ObtainThumbnailFromThirdPartyEncoder(const FString& InEncoderPath, const FString& InCurrentFile);
+	TOptional<FIntPoint> ObtainImageDimensionsFromThirdPartyEncoder(const FString& InEncoderPath, const FString& InCurrentFile);
+
+	TArray<uint8> RunProcess(const FString& InProcessName, const FString& InProcessArgs);
+	TArray<uint8> ReadFromPipe(void* InPipe);
 };
 
-}
\ No newline at end of file
+}
+
+#undef UE_API
\ No newline at end of file
diff --git a/Engine/Plugins/VirtualProduction/CaptureManager/CaptureManagerDevices/Source/VideoLiveLinkDeviceCommon/VideoLiveLinkDeviceCommon.Build.cs b/Engine/Plugins/VirtualProduction/CaptureManager/CaptureManagerDevices/Source/VideoLiveLinkDeviceCommon/VideoLiveLinkDeviceCommon.Build.cs
--- a/Engine/Plugins/VirtualProduction/CaptureManager/CaptureManagerDevices/Source/VideoLiveLinkDeviceCommon/VideoLiveLinkDeviceCommon.Build.cs
+++ b/Engine/Plugins/VirtualProduction/CaptureManager/CaptureManagerDevices/Source/VideoLiveLinkDeviceCommon/VideoLiveLinkDeviceCommon.Build.cs
@@ -12,20 +12,20 @@ public class VideoLiveLinkDeviceCommon : ModuleRules
 
 		PublicDependencyModuleNames.AddRange(new string[]
 		{
+			"CaptureManagerTakeMetadata",
 			"ImageCore",
-			"Core"
+			"Core",
 		});
 
 		PrivateDependencyModuleNames.AddRange(new string[]
 		{
 			"CaptureManagerMediaRW",
+			"CaptureManagerSettings",
 			"Engine",
 			"Slate",
 			"SlateCore",
 			"CoreUObject",
 			"UnrealEd"
 		});
-
-		CppCompileWarningSettings.UnsafeTypeCastWarningLevel = WarningLevel.Error;
 	}
 }
