Update stereo video ingest device, to automatically discover technoprops takes folder structure.

diff --git a/Engine/Plugins/VirtualProduction/CaptureManager/CaptureManagerDevices/Source/StereoVideoIngestDevice/Private/StereoVideoIngestDevice.cpp b/Engine/Plugins/VirtualProduction/CaptureManager/CaptureManagerDevices/Source/StereoVideoIngestDevice/Private/StereoVideoIngestDevice.cpp
--- a/Engine/Plugins/VirtualProduction/CaptureManager/CaptureManagerDevices/Source/StereoVideoIngestDevice/Private/StereoVideoIngestDevice.cpp
+++ b/Engine/Plugins/VirtualProduction/CaptureManager/CaptureManagerDevices/Source/StereoVideoIngestDevice/Private/StereoVideoIngestDevice.cpp
@@ -53,6 +53,12 @@ namespace UE::StereoVideoLiveLinkDevice::Private
 	{
 		return FModuleManager::LoadModuleChecked<FCaptureManagerMediaRWModule>("CaptureManagerMediaRW").Get();
 	}
+
+	static constexpr int32 MaximumVideoFilesCountPerTake = 2;
+	static constexpr int32 MaximumAudioFilesCountPerTake = 1;
+
+	static constexpr int32 DirectoriesCountToIterateForTakesSearch = 2000;
+	static constexpr int32 DirectoriesCountToIterateForVideoAudioFilesSearch = 500;
 }
 
 const UStereoVideoIngestDeviceSettings* UStereoVideoIngestDevice::GetSettings() const
@@ -163,9 +169,9 @@ void UStereoVideoIngestDevice::UpdateTakeList_Implementation(UIngestCapability_U
 		{
 			FTakeWithComponents& TakeWithComponents = TakesWithComponentsCandidates[SlateName];
 
-			bool TakeAsExpected = ((TakeWithComponents.CountComponents(ETakeComponentType::VIDEO) == 2 && TakeWithComponents.CountComponents(ETakeComponentType::IMAGE_SEQUENCE) == 0) ||
-									(TakeWithComponents.CountComponents(ETakeComponentType::VIDEO) == 0 && TakeWithComponents.CountComponents(ETakeComponentType::IMAGE_SEQUENCE) == 2)) &&
-				TakeWithComponents.CountComponents(ETakeComponentType::AUDIO) <= 1;
+			bool TakeAsExpected = ((TakeWithComponents.CountComponents(ETakeComponentType::VIDEO) == UE::StereoVideoLiveLinkDevice::Private::MaximumVideoFilesCountPerTake && TakeWithComponents.CountComponents(ETakeComponentType::IMAGE_SEQUENCE) == 0) ||
+				(TakeWithComponents.CountComponents(ETakeComponentType::VIDEO) == 0 && TakeWithComponents.CountComponents(ETakeComponentType::IMAGE_SEQUENCE) == UE::StereoVideoLiveLinkDevice::Private::MaximumVideoFilesCountPerTake)) &&
+				TakeWithComponents.CountComponents(ETakeComponentType::AUDIO) <= UE::StereoVideoLiveLinkDevice::Private::MaximumAudioFilesCountPerTake;
 
 			if (!TakeAsExpected)
 			{
@@ -195,13 +201,12 @@ TMap<FString, UStereoVideoIngestDevice::FTakeWithComponents> UStereoVideoIngestD
 	TMap<FString, FTakeWithComponents> TakesWithComponentsCandidates;
 
 	int32 DirectoriesInterrogatedCount = 0;
-	const int32 DirectoriesToInterrogateInOneRun = 200;
 
 	const bool bIterationResult = IFileManager::Get().IterateDirectoryRecursively(*StoragePath, [&](const TCHAR* InPath, bool bInIsDirectory)
 		{
 			if (bInIsDirectory)
 			{
-				if (++DirectoriesInterrogatedCount > DirectoriesToInterrogateInOneRun)
+				if (++DirectoriesInterrogatedCount > UE::StereoVideoLiveLinkDevice::Private::DirectoriesCountToIterateForTakesSearch)
 				{
 					return false;
 				}
@@ -239,7 +244,7 @@ void UStereoVideoIngestDevice::ExtractTakeComponents(FString ComponentPath, FStr
 	{
 		TakeWithComponents = ExtractTakeComponentsUsingTokens(ComponentPath, StoragePath, Format, ComponentType);
 	}
-	
+
 	if (!TakeWithComponents.IsSet())
 	{
 		FString ComponentRootPath;
@@ -259,7 +264,7 @@ TOptional<UStereoVideoIngestDevice::FTakeWithComponents> UStereoVideoIngestDevic
 	FString RelativePath = FString(InPathStringView.SubStr(StoragePath.Len(), InPathStringView.Len() - StoragePath.Len()));
 	FString Extension = FPaths::GetExtension(RelativePath);
 	FString RelativePathNoExtension = RelativePath.Mid(0, RelativePath.Len() - Extension.Len() - 1);
-	
+
 	FTakeDiscoveryExpressionParser TokenParser(Format, RelativePathNoExtension, UE::StereoVideoLiveLinkDevice::Private::Delimiters);
 
 	FString StorageLeafWithRelativePath = FPaths::GetPathLeaf(StoragePath) + RelativePathNoExtension;
@@ -299,22 +304,69 @@ TOptional<UStereoVideoIngestDevice::FTakeWithComponents> UStereoVideoIngestDevic
 
 	if (ComponentType == ETakeComponentType::VIDEO || ComponentType == ETakeComponentType::AUDIO)
 	{
-		// Use video file container folder name as TakeName
-		// this is for case where user navigate StoragePath to a container folder
-		TArray<FString> Parts;
-		FPaths::NormalizeDirectoryName(RelativePath);
-		RelativePath.ParseIntoArray(Parts, TEXT("/"));
-		if (Parts.Num() > 1)
+		// As a default, assume that user user navigates StoragePath to TakeFolder.
+		// Then use specified storage path leaf folder as TakeName
+		TakeName = FPaths::GetPathLeaf(StoragePath);
+		TakePath = StoragePath;
+
+
+		// Run a quick count of video and audio files in the selected StoragePath
+		bool bMoreAudioOrVideoFilesThanExpectedPerTake = false;
 		{
-			TakeName = Parts[0];
-			TakePath = StoragePath + "/" + TakeName;
+			int32 DirectoriesInterrogatedCount = 0;
+			int32 VideoFilesFound = 0;
+			int32 AudioFilesFound = 0;
+
+			const bool bIterationResult = IFileManager::Get().IterateDirectoryRecursively(*StoragePath, [&](const TCHAR* InPath, bool bInIsDirectory)
+				{
+					if (bInIsDirectory)
+					{
+						if (++DirectoriesInterrogatedCount > UE::StereoVideoLiveLinkDevice::Private::DirectoriesCountToIterateForVideoAudioFilesSearch)
+						{
+							return false;
+						}
+					}
+
+					FString Extension = FPaths::GetExtension(InPath);
+					if (UE::StereoVideoLiveLinkDevice::Private::SupportedVideoExtensions.Contains(Extension))
+					{
+						VideoFilesFound++;
+					}
+					else if (UE::StereoVideoLiveLinkDevice::Private::SupportedAudioExtensions.Contains(Extension))
+					{
+						AudioFilesFound++;
+					}
+
+					if (VideoFilesFound > UE::StereoVideoLiveLinkDevice::Private::MaximumVideoFilesCountPerTake 
+						|| AudioFilesFound > UE::StereoVideoLiveLinkDevice::Private::MaximumAudioFilesCountPerTake)
+					{
+						// Once we found more video or audio files than expected for one take, stop iteration over folders
+						bMoreAudioOrVideoFilesThanExpectedPerTake = true;
+						return false;
+					}
+
+					return true;
+				});
 		}
-		else
+
+		// If specified StoragePath contains more than maximum video files count or more than maximum audio files count,
+		// further assumption is that user selected a folder that is a container to take folders. 
+		if (bMoreAudioOrVideoFilesThanExpectedPerTake)
 		{
-			// If the above fails, because user navigates StoragePath to TakeFolder, 
-			// we could use specified storage path leaf folder as TakeName
-			TakeName = FPaths::GetPathLeaf(StoragePath);
-			TakePath = StoragePath;
+			// If sub folder exist, use it's name as TakeName
+			TArray<FString> Parts;
+			FPaths::NormalizeDirectoryName(RelativePath);
+			RelativePath.ParseIntoArray(Parts, TEXT("/"));
+			if (Parts.Num() > 1)
+			{
+				TakeName = Parts[0];
+				TakePath = StoragePath / TakeName;
+			}
+			else
+			{
+				// Invalid take folder
+				return {};
+			}
 		}
 	}
 	else if (ComponentType == ETakeComponentType::IMAGE_SEQUENCE)
@@ -524,7 +576,7 @@ TOptional<FTakeMetadata> UStereoVideoIngestDevice::CreateTakeMetadata(const FTak
 				}
 
 				Video.Path = TakeComponent.Path;
-				
+
 				FString MediaFilePath = FPaths::ConvertRelativePathToFull(InTakeWithComponents.TakePath, Video.Path);
 
 				FCaptureExtractVideoInfo::FResult ExtractorOpt = FCaptureExtractVideoInfo::Create(MediaFilePath);
