Resolving a crash when user removes the device during ingest

diff --git a/Engine/Plugins/VirtualProduction/CaptureManager/CaptureManagerDevices/Source/CPSLiveLinkDevice/Private/LiveLinkFaceDevice.cpp b/Engine/Plugins/VirtualProduction/CaptureManager/CaptureManagerDevices/Source/CPSLiveLinkDevice/Private/LiveLinkFaceDevice.cpp
--- a/Engine/Plugins/VirtualProduction/CaptureManager/CaptureManagerDevices/Source/CPSLiveLinkDevice/Private/LiveLinkFaceDevice.cpp
+++ b/Engine/Plugins/VirtualProduction/CaptureManager/CaptureManagerDevices/Source/CPSLiveLinkDevice/Private/LiveLinkFaceDevice.cpp
@@ -53,6 +53,11 @@ void ULiveLinkFaceDevice::OnDeviceAdded()
 
 void ULiveLinkFaceDevice::OnDeviceRemoved()
 {
+	if (Device)
+	{
+		Device->CancelAllExports();
+	}
+
 	ILiveLinkDeviceCapability_Connection::Execute_Disconnect(this);
 
 	Super::OnDeviceRemoved();
@@ -160,6 +165,11 @@ void ULiveLinkFaceDevice::RunConvertAndUploadTake(const UIngestCapability_Proces
 
 void ULiveLinkFaceDevice::CancelIngestProcess_Implementation(const UIngestCapability_ProcessHandle* InProcessHandle)
 {
+	if (!Device)
+	{
+		return;
+	}
+
 	UE::CaptureManager::FTakeId TakeId = InProcessHandle->GetTakeId();
 
 	Device->CancelExport(TakeId);
diff --git a/Engine/Plugins/VirtualProduction/CaptureManager/CaptureManagerDevices/Source/CPSLiveLinkDevice/Private/Protocol/CPSDevice.cpp b/Engine/Plugins/VirtualProduction/CaptureManager/CaptureManagerDevices/Source/CPSLiveLinkDevice/Private/Protocol/CPSDevice.cpp
--- a/Engine/Plugins/VirtualProduction/CaptureManager/CaptureManagerDevices/Source/CPSLiveLinkDevice/Private/Protocol/CPSDevice.cpp
+++ b/Engine/Plugins/VirtualProduction/CaptureManager/CaptureManagerDevices/Source/CPSLiveLinkDevice/Private/Protocol/CPSDevice.cpp
@@ -266,6 +266,11 @@ void FCPSDevice::CancelExport(FTakeId InTakeId)
 	}
 }
 
+void FCPSDevice::CancelAllExports()
+{
+	ExportClient->AbortAllExports();
+}
+
 void FCPSDevice::FetchThumbnailForTake(FTakeId InTakeId, TUniquePtr<FBaseStream> InStream)
 {
 	FetchFileForTake(InTakeId, MoveTemp(InStream), FString("thumbnail.jpg"));
diff --git a/Engine/Plugins/VirtualProduction/CaptureManager/CaptureManagerDevices/Source/CPSLiveLinkDevice/Public/Protocol/CPSDevice.h b/Engine/Plugins/VirtualProduction/CaptureManager/CaptureManagerDevices/Source/CPSLiveLinkDevice/Public/Protocol/CPSDevice.h
--- a/Engine/Plugins/VirtualProduction/CaptureManager/CaptureManagerDevices/Source/CPSLiveLinkDevice/Public/Protocol/CPSDevice.h
+++ b/Engine/Plugins/VirtualProduction/CaptureManager/CaptureManagerDevices/Source/CPSLiveLinkDevice/Public/Protocol/CPSDevice.h
@@ -95,6 +95,7 @@ public:
 
 	void StartExport(FTakeId InTakeId, TUniquePtr<FBaseStream> InStream);
 	void CancelExport(FTakeId InTakeId);
+	void CancelAllExports();
 
 	/** Fetches the thumbnail for a single take. */
 	void FetchThumbnailForTake(FTakeId InTakeId, TUniquePtr<FBaseStream> InStream);
