Put quotes around file paths passed to ffmpeg and ffprobe

diff --git a/Engine/Plugins/VirtualProduction/CaptureManager/CaptureManagerApp/Source/IngestLiveLinkDevice/Private/Utils/CaptureExtractTimecode.cpp b/Engine/Plugins/VirtualProduction/CaptureManager/CaptureManagerApp/Source/IngestLiveLinkDevice/Private/Utils/CaptureExtractTimecode.cpp
--- a/Engine/Plugins/VirtualProduction/CaptureManager/CaptureManagerApp/Source/IngestLiveLinkDevice/Private/Utils/CaptureExtractTimecode.cpp
+++ b/Engine/Plugins/VirtualProduction/CaptureManager/CaptureManagerApp/Source/IngestLiveLinkDevice/Private/Utils/CaptureExtractTimecode.cpp
@@ -274,7 +274,7 @@ FCaptureExtractVideoInfo::FExtractResult FCaptureExtractVideoInfo::ExtractInfoUs
 {
 	using namespace UE::CaptureManager;
 
-	FString CommandArgs = TEXT("-v error -select_streams v:0 -show_entries stream_tags=timecode:stream=r_frame_rate -of default=noprint_wrappers=1:nokey=1 ") + InFilePath;
+	FString CommandArgs = FString::Format(TEXT("-v error -select_streams v:0 -show_entries stream_tags=timecode:stream=r_frame_rate -of default=noprint_wrappers=1:nokey=1 \"{0}\""), { InFilePath });
 
 	void* ReadPipe = nullptr;
 	void* WritePipe = nullptr;
diff --git a/Engine/Plugins/VirtualProduction/CaptureManager/CaptureManagerDevices/Source/VideoLiveLinkDeviceCommon/Private/Utils/VideoDeviceThumbnailExtractor.cpp b/Engine/Plugins/VirtualProduction/CaptureManager/CaptureManagerDevices/Source/VideoLiveLinkDeviceCommon/Private/Utils/VideoDeviceThumbnailExtractor.cpp
--- a/Engine/Plugins/VirtualProduction/CaptureManager/CaptureManagerDevices/Source/VideoLiveLinkDeviceCommon/Private/Utils/VideoDeviceThumbnailExtractor.cpp
+++ b/Engine/Plugins/VirtualProduction/CaptureManager/CaptureManagerDevices/Source/VideoLiveLinkDeviceCommon/Private/Utils/VideoDeviceThumbnailExtractor.cpp
@@ -228,7 +228,7 @@ TOptional<FIntPoint> FVideoDeviceThumbnailExtractor::ObtainImageDimensionsFromTh
 		return {};
 	}
 
-	FString CommandLine = FString::Format(TEXT("-v error -select_streams v:0 -show_entries stream=width,height -of default=noprint_wrappers=1 {0}"), { InCurrentFile });
+	FString CommandLine = FString::Format(TEXT("-v error -select_streams v:0 -show_entries stream=width,height -of default=noprint_wrappers=1 \"{0}\""), { InCurrentFile });
 	FString EncoderPath = Private::GetFFprobePath(InEncoderPath);
 
 	TArray<uint8> WidthAndHeight = RunProcess(EncoderPath, CommandLine);
