Fix crash in Live Link Hub when gathering thumbnails of video files containing rotation metadata via the third party encoder within Capture Manager

[RN] fix virtual production/tools

diff --git a/Engine/Plugins/VirtualProduction/CaptureManager/CaptureManagerDevices/Source/VideoLiveLinkDeviceCommon/Private/Utils/VideoDeviceThumbnailExtractor.cpp b/Engine/Plugins/VirtualProduction/CaptureManager/CaptureManagerDevices/Source/VideoLiveLinkDeviceCommon/Private/Utils/VideoDeviceThumbnailExtractor.cpp
--- a/Engine/Plugins/VirtualProduction/CaptureManager/CaptureManagerDevices/Source/VideoLiveLinkDeviceCommon/Private/Utils/VideoDeviceThumbnailExtractor.cpp
+++ b/Engine/Plugins/VirtualProduction/CaptureManager/CaptureManagerDevices/Source/VideoLiveLinkDeviceCommon/Private/Utils/VideoDeviceThumbnailExtractor.cpp
@@ -206,7 +206,7 @@ TArray<FColor> FVideoDeviceThumbnailExtractor::ObtainThumbnailFromThirdPartyEnco
 {
 	FString ScaleArg = FString::Printf(TEXT("scale=%d:-1"), ThirdPartyEncoderThumbnailWidth);
 
-	FString CommandLineArgs = FString::Format(TEXT("-hide_banner -loglevel error -i \"{0}\" -map v:0 -vf \"thumbnail=2, {1}\" -frames:v 1 -f rawvideo -pix_fmt rgb24 -"), { InCurrentFile, ScaleArg });
+	FString CommandLineArgs = FString::Format(TEXT("-noautorotate -hide_banner -loglevel error -i \"{0}\" -map v:0 -vf \"thumbnail=2, {1}\" -frames:v 1 -f rawvideo -pix_fmt rgb24 -"), { InCurrentFile, ScaleArg });
 	
 	TArray<uint8> ThumbnailRawData = RunProcess(InEncoderPath, CommandLineArgs);
 
