[CaptureManager] Fixed incorrect audio timecode bug for the Stereo Video Ingest device

diff --git a/Engine/Plugins/VirtualProduction/CaptureManager/CaptureManagerDevices/Source/StereoVideoIngestDevice/Private/StereoVideoIngestDevice.cpp b/Engine/Plugins/VirtualProduction/CaptureManager/CaptureManagerDevices/Source/StereoVideoIngestDevice/Private/StereoVideoIngestDevice.cpp
--- a/Engine/Plugins/VirtualProduction/CaptureManager/CaptureManagerDevices/Source/StereoVideoIngestDevice/Private/StereoVideoIngestDevice.cpp
+++ b/Engine/Plugins/VirtualProduction/CaptureManager/CaptureManagerDevices/Source/StereoVideoIngestDevice/Private/StereoVideoIngestDevice.cpp
@@ -523,149 +523,175 @@ TOptional<FTakeMetadata> UStereoVideoIngestDevice::CreateTakeMetadata(const FTak
 
 	TakeMetadata.Device.Model = TEXT("StereoHMC");
 
-	int32 AudioNameCounter = 1;
+
+	TArray<FTakeWithComponents::Component> VideoComponents = InTakeWithComponents.Components.FilterByPredicate([](const FTakeWithComponents::Component& Component)
+	{
+		return Component.Type == ETakeComponentType::VIDEO || Component.Type == ETakeComponentType::IMAGE_SEQUENCE;
+	});
+
+	check(VideoComponents.Num() == UE::StereoVideoLiveLinkDevice::Private::MaximumVideoFilesCountPerTake);
 
 	bool bVideoFrameRateSet = false;
 	FFrameRate VideoFrameRate;
 
-	for (const FTakeWithComponents::Component& TakeComponent : InTakeWithComponents.Components)
+	for (const FTakeWithComponents::Component& VideoComponent : VideoComponents)
 	{
-		if (TakeComponent.Type == ETakeComponentType::VIDEO || TakeComponent.Type == ETakeComponentType::IMAGE_SEQUENCE)
+		FTakeMetadata::FVideo Video;
+
+		FString Name = VideoComponent.Name;
+		if (Name.IsEmpty())
 		{
-			FTakeMetadata::FVideo Video;
+			Name = VideoComponent.Path.Mid(InTakeWithComponents.TakePath.Len() + 1);
+			SanitizePackagePath(Name, '_');
+		}
 
-			FString Name = TakeComponent.Name;
-			if (Name.IsEmpty())
-			{
-				Name = TakeComponent.Path.Mid(InTakeWithComponents.TakePath.Len() + 1);
-				SanitizePackagePath(Name, '_');
-			}
+		Video.Name = MoveTemp(Name);
 
-			Video.Name = MoveTemp(Name);
-			Video.Format = FPaths::GetExtension(TakeComponent.Path);
+		if (VideoComponent.Type == ETakeComponentType::VIDEO)
+		{
+			Video.Format = FPaths::GetExtension(VideoComponent.Path);
+			Video.PathType = FTakeMetadata::FVideo::EPathType::File;
 
-			if (TakeComponent.Type == ETakeComponentType::VIDEO)
+			TValueOrError<TUniquePtr<IVideoReader>, FText> VideoReaderResult = MediaManager.CreateVideoReader(VideoComponent.Path);
+			if (VideoReaderResult.HasValue())
 			{
-				Video.PathType = FTakeMetadata::FVideo::EPathType::File;
-
-				TValueOrError<TUniquePtr<IVideoReader>, FText> VideoReaderResult = MediaManager.CreateVideoReader(TakeComponent.Path);
-				if (VideoReaderResult.HasValue())
-				{
-					TUniquePtr<IVideoReader> VideoReader = VideoReaderResult.StealValue();
-					Video.FrameRate = static_cast<float>(VideoReader->GetFrameRate().AsDecimal());
+				TUniquePtr<IVideoReader> VideoReader = VideoReaderResult.StealValue();
+				Video.FrameRate = static_cast<float>(VideoReader->GetFrameRate().AsDecimal());
 
-					if (!bVideoFrameRateSet)
-					{
-						VideoFrameRate = UE::CaptureManager::ParseFrameRate(Video.FrameRate);
-					}
-				}
-				else
+				if (!bVideoFrameRateSet)
 				{
-					UE_LOG(LogVideoLiveLinkDevice, Warning, TEXT("Failed to determine the frame rate for the video file %s. Consider enabling Third Party Encoder in Capture Manager settings."), *TakeComponent.Path);
+					VideoFrameRate = UE::CaptureManager::ParseFrameRate(Video.FrameRate);
+					bVideoFrameRateSet = true;
 				}
+			}
+			else
+			{
+				UE_LOG(LogVideoLiveLinkDevice, Warning, TEXT("Failed to determine the frame rate for the video file %s. Consider enabling Third Party Encoder in Capture Manager settings."), *VideoComponent.Path);
+			}
 
-				if (!TakeMetadata.Thumbnail.GetRawImage().IsSet())
-				{
-					FVideoDeviceThumbnailExtractor ThumbnailExtractor;
-					TOptional<FTakeThumbnailData::FRawImage> RawImageOpt = ThumbnailExtractor.ExtractThumbnail(TakeComponent.Path);
+			if (!TakeMetadata.Thumbnail.GetRawImage().IsSet())
+			{
+				FVideoDeviceThumbnailExtractor ThumbnailExtractor;
+				TOptional<FTakeThumbnailData::FRawImage> RawImageOpt = ThumbnailExtractor.ExtractThumbnail(VideoComponent.Path);
 
-					if (RawImageOpt.IsSet())
-					{
-						TakeMetadata.Thumbnail = MoveTemp(RawImageOpt.GetValue());
-					}
+				if (RawImageOpt.IsSet())
+				{
+					TakeMetadata.Thumbnail = MoveTemp(RawImageOpt.GetValue());
 				}
+			}
 
-				Video.Path = TakeComponent.Path;
+			Video.Path = VideoComponent.Path;
 
-				FString MediaFilePath = FPaths::ConvertRelativePathToFull(InTakeWithComponents.TakePath, Video.Path);
+			FString MediaFilePath = FPaths::ConvertRelativePathToFull(InTakeWithComponents.TakePath, Video.Path);
 
-				FCaptureExtractVideoInfo::FResult ExtractorOpt = FCaptureExtractVideoInfo::Create(MediaFilePath);
-				if (ExtractorOpt.IsValid())
+			FCaptureExtractVideoInfo::FResult ExtractorOpt = FCaptureExtractVideoInfo::Create(MediaFilePath);
+			if (ExtractorOpt.IsValid())
+			{
+				FCaptureExtractVideoInfo Extractor = ExtractorOpt.StealValue();
+				Video.TimecodeStart = Extractor.GetTimecode().ToString();
+				if (FMath::IsNearlyZero(Video.FrameRate))
 				{
-					FCaptureExtractVideoInfo Extractor = ExtractorOpt.StealValue();
-					Video.TimecodeStart = Extractor.GetTimecode().ToString();
-					if (FMath::IsNearlyZero(Video.FrameRate))
-					{
-						Video.FrameRate = static_cast<float>(Extractor.GetFrameRate().AsDecimal());
-					}
+					Video.FrameRate = static_cast<float>(Extractor.GetFrameRate().AsDecimal());
+					VideoFrameRate = Extractor.GetFrameRate();
+					bVideoFrameRateSet = true;
 				}
 			}
 			else
 			{
-				Video.PathType = FTakeMetadata::FVideo::EPathType::Folder;
+				UE_LOG(LogVideoLiveLinkDevice, Warning, TEXT("Failed to determine the timecode for the video file %s. Consider enabling Third Party Encoder in Capture Manager settings."), *VideoComponent.Path);
+			}
+		}
+		else // ETakeComponentType::IMAGE_SEQUENCE
+		{
+			Video.PathType = FTakeMetadata::FVideo::EPathType::Folder;
 
-				int32 FramesCount = 0;
-				bool FirstFrame = true;
+			int32 FramesCount = 0;
+			bool FirstFrame = true;
 
-				bool bIterationResult = IFileManager::Get().IterateDirectoryRecursively(*TakeComponent.Path, [&](const TCHAR* InPath, bool bInIsDirectory)
+			bool bIterationResult = IFileManager::Get().IterateDirectoryRecursively(*VideoComponent.Path, [&](const TCHAR* InPath, bool bInIsDirectory)
+				{
+					if (FirstFrame)
 					{
-						if (FirstFrame)
-						{
-							FirstFrame = false;
+						FirstFrame = false;
 
-							if (!TakeMetadata.Thumbnail.GetThumbnailData().IsSet())
+						if (!TakeMetadata.Thumbnail.GetThumbnailData().IsSet())
+						{
+							TArray<uint8> ThumbnailData;
+							if (FFileHelper::LoadFileToArray(ThumbnailData, InPath))
 							{
-								TArray<uint8> ThumbnailData;
-								if (FFileHelper::LoadFileToArray(ThumbnailData, InPath))
-								{
-									TakeMetadata.Thumbnail = ThumbnailData;
-								}
+								TakeMetadata.Thumbnail = ThumbnailData;
 							}
+						}
 
-							FString Path, Filename, Extension;
-							FPaths::Split(InPath, Path, Filename, Extension);
-							Video.Format = Extension;
+						FString Path, Filename, Extension;
+						FPaths::Split(InPath, Path, Filename, Extension);
+						Video.Format = Extension;
 
-							Video.Path = TakeComponent.Path;
+						Video.Path = VideoComponent.Path;
+						{
+							FImage Image;
+							FImageUtils::LoadImage(InPath, Image);
+							if (Image.IsImageInfoValid())
 							{
-								FImage Image;
-								FImageUtils::LoadImage(InPath, Image);
-								if (Image.IsImageInfoValid())
-								{
-									Video.FrameWidth = Image.SizeX;
-									Video.FrameHeight = Image.SizeY;
-								}
+								Video.FrameWidth = Image.SizeX;
+								Video.FrameHeight = Image.SizeY;
 							}
 						}
+					}
 
-						FramesCount++;
-
-						return true;
-					});
+					FramesCount++;
 
-				Video.FramesCount = FramesCount;
-			}
+					return true;
+				});
 
-			TakeMetadata.Video.Add(MoveTemp(Video));
+			Video.FramesCount = FramesCount;
 		}
-		else if (TakeComponent.Type == ETakeComponentType::AUDIO)
+
+		TakeMetadata.Video.Add(MoveTemp(Video));
+	}
+
+
+	TArray<FTakeWithComponents::Component> AudioComponents = InTakeWithComponents.Components.FilterByPredicate([](const FTakeWithComponents::Component& Component)
+	{
+		return Component.Type == ETakeComponentType::AUDIO;
+	});
+
+	int32 AudioNameCounter = 1;
+
+	for (const FTakeWithComponents::Component& AudioComponent : AudioComponents)
+	{	
+		FTakeMetadata::FAudio Audio;
+		FString Name = AudioComponent.Name;
+		if (Name.IsEmpty())
 		{
-			FTakeMetadata::FAudio Audio;
-			FString Name = TakeComponent.Name;
-			if (Name.IsEmpty())
-			{
-				Name = TEXT("audio") + FString::FromInt(AudioNameCounter++);
-			}
-			Audio.Name = Name;
-			Audio.Path = TakeComponent.Path;
-			Audio.Duration = 0;
+			Name = TEXT("audio") + FString::FromInt(AudioNameCounter++);
+		}
+		Audio.Name = Name;
+		Audio.Path = AudioComponent.Path;
+		Audio.Duration = 0;
 
-			FString MediaFilePath = FPaths::ConvertRelativePathToFull(InTakeWithComponents.TakePath, Audio.Path);
-			TSharedPtr<FCaptureExtractAudioTimecode> Extractor = MakeShareable(new FCaptureExtractAudioTimecode(MediaFilePath));
+		FString MediaFilePath = FPaths::ConvertRelativePathToFull(InTakeWithComponents.TakePath, Audio.Path);
+		TSharedPtr<FCaptureExtractAudioTimecode> Extractor = MakeShareable(new FCaptureExtractAudioTimecode(MediaFilePath));
 
-			// The video frame rate will be used to calculate the timecode rate if the timecode rate cannot be extracted from the audio file
-			FCaptureExtractAudioTimecode::FTimecodeInfoResult TimecodeInfoResult = Extractor->Extract(VideoFrameRate);
-			if (TimecodeInfoResult.IsValid())
-			{
-				FTimecodeInfo TimecodeInfo = TimecodeInfoResult.GetValue();
+		// The video frame rate will be used to calculate the timecode rate if the timecode rate cannot be extracted from the audio file
+		// This is why we need to iterate over the video and audio components independently
+		if (!bVideoFrameRateSet)
+		{
+			UE_LOG(LogVideoLiveLinkDevice, Warning, TEXT("Target frame rate from the input video(s) not found. This is used to set the timecode rate for the audio. Audio timecode may be set incorrectly as a result."));
+		}
 
-				Audio.TimecodeStart = TimecodeInfo.Timecode.ToString();
-				Audio.TimecodeRate = static_cast<float>(TimecodeInfo.TimecodeRate.AsDecimal());
-			}
+		FCaptureExtractAudioTimecode::FTimecodeInfoResult TimecodeInfoResult = Extractor->Extract(VideoFrameRate);
+		if (TimecodeInfoResult.IsValid())
+		{
+			FTimecodeInfo TimecodeInfo = TimecodeInfoResult.GetValue();
 
-			TakeMetadata.Audio.Add(MoveTemp(Audio));
+			Audio.TimecodeStart = TimecodeInfo.Timecode.ToString();
+			Audio.TimecodeRate = static_cast<float>(TimecodeInfo.TimecodeRate.AsDecimal());
 		}
+
+		TakeMetadata.Audio.Add(MoveTemp(Audio));
 	}
+
 	return TakeMetadata;
 }
 
