Fixing a crash when removing device during connection

diff --git a/Engine/Plugins/VirtualProduction/CaptureManager/CaptureManagerDevices/Source/CPSLiveLinkDevice/Private/Protocol/CPSDevice.cpp b/Engine/Plugins/VirtualProduction/CaptureManager/CaptureManagerDevices/Source/CPSLiveLinkDevice/Private/Protocol/CPSDevice.cpp
--- a/Engine/Plugins/VirtualProduction/CaptureManager/CaptureManagerDevices/Source/CPSLiveLinkDevice/Private/Protocol/CPSDevice.cpp
+++ b/Engine/Plugins/VirtualProduction/CaptureManager/CaptureManagerDevices/Source/CPSLiveLinkDevice/Private/Protocol/CPSDevice.cpp
@@ -260,6 +260,11 @@ void FCPSDevice::StartExport(FTakeId InTakeId, TUniquePtr<FBaseStream> InStream)
 
 void FCPSDevice::CancelExport(FTakeId InTakeId)
 {
+	if (!ExportClient)
+	{
+		return;
+	}
+
 	if (FExportClient::FTaskId* TaskId = IdMap.Find(InTakeId))
 	{
 		ExportClient->AbortExport(*TaskId);
@@ -268,7 +273,10 @@ void FCPSDevice::CancelExport(FTakeId InTakeId)
 
 void FCPSDevice::CancelAllExports()
 {
-	ExportClient->AbortAllExports();
+	if (ExportClient)
+	{
+		ExportClient->AbortAllExports();
+	}
 }
 
 void FCPSDevice::FetchThumbnailForTake(FTakeId InTakeId, TUniquePtr<FBaseStream> InStream)
@@ -311,6 +319,7 @@ void FCPSDevice::FetchFileForTake(FTakeId InTakeId, TUniquePtr<FBaseStream> InSt
 		TakeFile.Length = Length;
 		TakeFile.Offset = 0;
 
+		check(ExportClient);
 		ExportClient->ExportTakeFiles(TakeObject.Name, { MoveTemp(TakeFile) }, MoveTemp(InStream));
 	}
 }
@@ -355,6 +364,7 @@ void FCPSDevice::FetchFiles(TUniquePtr<FBaseStream> InStream, TArray<FString> In
 	// Unlock earlier as there is no need to lock the ExporFiles function
 	Lock.Unlock();
 
+	check(ExportClient);
 	ExportClient->ExportFiles(MoveTemp(TakeFileArrayMap), MoveTemp(InStream));
 }
 
