[Capture Manager] Moved some ingest tasks off the foreground threads.

#rnx

diff --git a/Engine/Plugins/VirtualProduction/CaptureManager/CaptureManagerApp/Source/IngestLiveLinkDevice/Private/BaseIngestLiveLinkDevice.cpp b/Engine/Plugins/VirtualProduction/CaptureManager/CaptureManagerApp/Source/IngestLiveLinkDevice/Private/BaseIngestLiveLinkDevice.cpp
--- a/Engine/Plugins/VirtualProduction/CaptureManager/CaptureManagerApp/Source/IngestLiveLinkDevice/Private/BaseIngestLiveLinkDevice.cpp
+++ b/Engine/Plugins/VirtualProduction/CaptureManager/CaptureManagerApp/Source/IngestLiveLinkDevice/Private/BaseIngestLiveLinkDevice.cpp
@@ -19,6 +19,8 @@ namespace UE::CaptureManager::Private
 {
 static void TryCleanUpWorkingDirectory(const FString& InWorkingDirectory)
 {
+	TRACE_CPUPROFILER_EVENT_SCOPE(FBaseIngestDevice_CleanUpWorkingDirectory);
+
 	FString NormalizedFullPath = FPaths::ConvertRelativePathToFull(InWorkingDirectory);
 	FPaths::NormalizeFilename(NormalizedFullPath);
 
diff --git a/Engine/Plugins/VirtualProduction/CaptureManager/CaptureManagerDevices/Source/CPSLiveLinkDevice/Private/LiveLinkFaceDevice.cpp b/Engine/Plugins/VirtualProduction/CaptureManager/CaptureManagerDevices/Source/CPSLiveLinkDevice/Private/LiveLinkFaceDevice.cpp
--- a/Engine/Plugins/VirtualProduction/CaptureManager/CaptureManagerDevices/Source/CPSLiveLinkDevice/Private/LiveLinkFaceDevice.cpp
+++ b/Engine/Plugins/VirtualProduction/CaptureManager/CaptureManagerDevices/Source/CPSLiveLinkDevice/Private/LiveLinkFaceDevice.cpp
@@ -19,6 +19,8 @@
 
 #include "HAL/FileManager.h"
 
+DEFINE_LOG_CATEGORY_STATIC(LogLiveLinkFaceDevice, Log, All);
+
 const ULiveLinkFaceDeviceSettings* ULiveLinkFaceDevice::GetSettings() const
 {
 	return GetDeviceSettings<ULiveLinkFaceDeviceSettings>();
@@ -142,25 +144,35 @@ void ULiveLinkFaceDevice::RunConvertAndUploadTake(const UIngestCapability_Proces
 	TStrongObjectPtr<const UIngestCapability_Options> IngestOptions(InIngestOptions);
 
 	TSharedPtr<FTaskProgress> TaskProgress = MakeShared<FTaskProgress>
-		(NumberOfTasks, FTaskProgress::FProgressReporter::CreateLambda([this, ProcessHandle](double InProgress)
-		 {
-			 ExecuteProcessProgressReporter(ProcessHandle.Get(), InProgress);
-		 }));
+		(NumberOfTasks, FTaskProgress::FProgressReporter::CreateLambda([WeakThis = TWeakObjectPtr<ULiveLinkFaceDevice>(this), ProcessHandle](double InProgress)
+			{
+				TStrongObjectPtr<ULiveLinkFaceDevice> This = WeakThis.Pin();
 
-	// Free the current thread that is waiting on next download
-	TWeakObjectPtr<ULiveLinkFaceDevice> WeakThis(this);
-	AsyncTask(ENamedThreads::Type::AnyThread, [WeakThis = MoveTemp(WeakThis), ProcessHandle, IngestOptions, TaskProgress = MoveTemp(TaskProgress)]()
-	{
-		TStrongObjectPtr<ULiveLinkFaceDevice> StrongThis = WeakThis.Pin();
+				if (!This.IsValid())
+				{
+					// Don't log this one, just because it's called so frequently
+					return;
+				}
 
-		if (!StrongThis.IsValid())
+				This->ExecuteProcessProgressReporter(ProcessHandle.Get(), InProgress);
+			}));
+
+	// Free the current thread that is waiting on next download
+	auto Task = [WeakThis = TWeakObjectPtr<ULiveLinkFaceDevice>(this), ProcessHandle, IngestOptions, TaskProgress = MoveTemp(TaskProgress)]()
 		{
-			return;
-		}
+			TStrongObjectPtr<ULiveLinkFaceDevice> This = WeakThis.Pin();
+
+			if (!This.IsValid())
+			{
+				UE_LOG(LogLiveLinkFaceDevice, Warning, TEXT("Failed to ingest take, the device has been destroyed"));
+				return;
+			}
 
-		StrongThis->Super::IngestTake(ProcessHandle.Get(), IngestOptions.Get(), TaskProgress);
-		StrongThis->RemoveDownloadedTakeData(ProcessHandle->GetTakeId());
-	});
+			This->Super::IngestTake(ProcessHandle.Get(), IngestOptions.Get(), TaskProgress);
+			This->RemoveDownloadedTakeData(ProcessHandle->GetTakeId());
+		};
+
+	UE::Tasks::Launch(UE_SOURCE_LOCATION, MoveTemp(Task), UE::Tasks::ETaskPriority::BackgroundNormal);
 }
 
 void ULiveLinkFaceDevice::CancelIngestProcess_Implementation(const UIngestCapability_ProcessHandle* InProcessHandle)
@@ -238,17 +250,27 @@ bool ULiveLinkFaceDevice::Connect_Implementation()
 
 bool ULiveLinkFaceDevice::Disconnect_Implementation()
 {
-	AsyncTask(ENamedThreads::AnyThread, [this]()
-	{
-		if (Device)
+	auto Task = [WeakThis = TWeakObjectPtr<ULiveLinkFaceDevice>(this)]()
 		{
-			Device->Stop();
-			Device->UnsubscribeAll();
-			bIsConnecting = false;
-			Device = nullptr;
-		}
-	});
-	
+			TStrongObjectPtr<ULiveLinkFaceDevice> This = WeakThis.Pin();
+
+			if (!This.IsValid())
+			{
+				UE_LOG(LogLiveLinkFaceDevice, Warning, TEXT("Failed to disconnect, the device has been destroyed"));
+				return;
+			}
+
+			if (This->Device)
+			{
+				This->Device->Stop();
+				This->Device->UnsubscribeAll();
+				This->bIsConnecting = false;
+				This->Device = nullptr;
+			}
+		};
+
+	UE::Tasks::Launch(UE_SOURCE_LOCATION, MoveTemp(Task), UE::Tasks::ETaskPriority::BackgroundNormal);
+
 	return true;
 }
 
diff --git a/Engine/Plugins/VirtualProduction/CaptureManager/CaptureManagerDevices/Source/MonoVideoIngestDevice/Private/MonoVideoIngestDevice.cpp b/Engine/Plugins/VirtualProduction/CaptureManager/CaptureManagerDevices/Source/MonoVideoIngestDevice/Private/MonoVideoIngestDevice.cpp
--- a/Engine/Plugins/VirtualProduction/CaptureManager/CaptureManagerDevices/Source/MonoVideoIngestDevice/Private/MonoVideoIngestDevice.cpp
+++ b/Engine/Plugins/VirtualProduction/CaptureManager/CaptureManagerDevices/Source/MonoVideoIngestDevice/Private/MonoVideoIngestDevice.cpp
@@ -103,67 +103,96 @@ FString UMonoVideoIngestDevice::GetFullTakePath(UE::CaptureManager::FTakeId InTa
 
 void UMonoVideoIngestDevice::RunConvertAndUploadTake(const UIngestCapability_ProcessHandle* InProcessHandle, const UIngestCapability_Options* InIngestOptions)
 {
-	AsyncTask(ENamedThreads::AnyThread, [this,
-			  ProcessHandle = TStrongObjectPtr<const UIngestCapability_ProcessHandle>(InProcessHandle),
-			  IngestOptions = TStrongObjectPtr<const UIngestCapability_Options>(InIngestOptions)]()
-	{
-		static constexpr uint32 NumberOfTasks = 2; // Convert, Upload
-
-		using namespace UE::CaptureManager;
-		TSharedPtr<FTaskProgress> TaskProgress
-			= MakeShared<FTaskProgress>(NumberOfTasks, FTaskProgress::FProgressReporter::CreateLambda([this, ProcessHandle](double InProgress)
+	auto Task = [WeakThis = TWeakObjectPtr<UMonoVideoIngestDevice>(this),
+		ProcessHandle = TStrongObjectPtr<const UIngestCapability_ProcessHandle>(InProcessHandle),
+		IngestOptions = TStrongObjectPtr<const UIngestCapability_Options>(InIngestOptions)]()
 		{
-			ExecuteProcessProgressReporter(ProcessHandle.Get(), InProgress);
-		}));
+			TStrongObjectPtr<UMonoVideoIngestDevice> This = WeakThis.Pin();
+
+			if (!This.IsValid())
+			{
+				UE_LOG(LogVideoLiveLinkDevice, Warning, TEXT("Failed to ingest take, the device has been destroyed"));
+				return;
+			}
+
+			static constexpr uint32 NumberOfTasks = 2; // Convert, Upload
+
+			using namespace UE::CaptureManager;
+			TSharedPtr<FTaskProgress> TaskProgress
+				= MakeShared<FTaskProgress>(NumberOfTasks, FTaskProgress::FProgressReporter::CreateLambda([WeakThis = WeakThis, ProcessHandle](double InProgress)
+					{
+						TStrongObjectPtr<UMonoVideoIngestDevice> This = WeakThis.Pin();
 
-		Super::IngestTake(ProcessHandle.Get(), IngestOptions.Get(), MoveTemp(TaskProgress));
-	});
+						if (!This.IsValid())
+						{
+							// Don't log this one, just because it's called so frequently
+							return;
+						}
+
+						This->ExecuteProcessProgressReporter(ProcessHandle.Get(), InProgress);
+					}));
+
+			This->Super::IngestTake(ProcessHandle.Get(), IngestOptions.Get(), MoveTemp(TaskProgress));
+		};
+
+	UE::Tasks::Launch(UE_SOURCE_LOCATION, MoveTemp(Task), UE::Tasks::ETaskPriority::BackgroundNormal);
 }
 
 void UMonoVideoIngestDevice::UpdateTakeList_Implementation(UIngestCapability_UpdateTakeListCallback* InCallback)
 {
-	AsyncTask(ENamedThreads::AnyThread, [this, Callback = TStrongObjectPtr<UIngestCapability_UpdateTakeListCallback>(InCallback)]
-	{
-		RemoveAllTakes();
-
-		FString StoragePath = GetSettings()->TakeDirectory.Path;
-		int32 DirectoriesInterrogatedCount = 0;
-		const int32 DirectoriesToInterrogateInOneRun = 200;
-	
-		TArray<FString> SupportedVideoFiles;
-		const bool bIterationResult = IFileManager::Get().IterateDirectoryRecursively(*StoragePath, [&](const TCHAR* InFileNameOrDirectory, bool bInIsDirectory)
+	auto Task = [WeakThis = TWeakObjectPtr<UMonoVideoIngestDevice>(this), Callback = TStrongObjectPtr<UIngestCapability_UpdateTakeListCallback>(InCallback)]
 		{
-			if (bInIsDirectory)
-			{
-				if (++DirectoriesInterrogatedCount > DirectoriesToInterrogateInOneRun)
-				{
-					return false;
-				}
-			}
+			TStrongObjectPtr<UMonoVideoIngestDevice> This = WeakThis.Pin();
 
-			if (!bInIsDirectory && IsVideoFile(InFileNameOrDirectory))
-			{			
-				SupportedVideoFiles.Add(InFileNameOrDirectory);
+			if (!This.IsValid())
+			{
+				UE_LOG(LogVideoLiveLinkDevice, Warning, TEXT("Failed to update take list, the device has been destroyed"));
+				return;
 			}
 
-			return true;
-		});
+			This->RemoveAllTakes();
 
-		for (const FString& SupportedVideoFile : SupportedVideoFiles)
-		{
-			TOptional<FTakeMetadata> TakeInfoResult = ReadTake(SupportedVideoFile);
+			FString StoragePath = This->GetSettings()->TakeDirectory.Path;
+			int32 DirectoriesInterrogatedCount = 0;
+			const int32 DirectoriesToInterrogateInOneRun = 200;
 
-			if (TakeInfoResult.IsSet())
+			TArray<FString> SupportedVideoFiles;
+			const bool bIterationResult = IFileManager::Get().IterateDirectoryRecursively(*StoragePath, [&](const TCHAR* InFileNameOrDirectory, bool bInIsDirectory)
+				{
+					if (bInIsDirectory)
+					{
+						if (++DirectoriesInterrogatedCount > DirectoriesToInterrogateInOneRun)
+						{
+							return false;
+						}
+					}
+
+					if (!bInIsDirectory && This->IsVideoFile(InFileNameOrDirectory))
+					{
+						SupportedVideoFiles.Add(InFileNameOrDirectory);
+					}
+
+					return true;
+				});
+
+			for (const FString& SupportedVideoFile : SupportedVideoFiles)
 			{
-				int32 TakeId = AddTake(TakeInfoResult.GetValue());
-				FullTakePaths.Add(TakeId, SupportedVideoFile);
+				TOptional<FTakeMetadata> TakeInfoResult = This->ReadTake(SupportedVideoFile);
 
-				PublishEvent<FIngestCapability_TakeAddedEvent>(TakeId);
+				if (TakeInfoResult.IsSet())
+				{
+					int32 TakeId = This->AddTake(TakeInfoResult.GetValue());
+					This->FullTakePaths.Add(TakeId, SupportedVideoFile);
+
+					This->PublishEvent<FIngestCapability_TakeAddedEvent>(TakeId);
+				}
 			}
-		}
 
-		ExecuteUpdateTakeListCallback(Callback.Get(), Execute_GetTakeIdentifiers(this));
-	});
+			This->ExecuteUpdateTakeListCallback(Callback.Get(), Execute_GetTakeIdentifiers(This.Get()));
+		};
+
+
+	UE::Tasks::Launch(UE_SOURCE_LOCATION, MoveTemp(Task), UE::Tasks::ETaskPriority::BackgroundNormal);
 }
 
 ELiveLinkDeviceConnectionStatus UMonoVideoIngestDevice::GetConnectionStatus_Implementation() const
diff --git a/Engine/Plugins/VirtualProduction/CaptureManager/CaptureManagerDevices/Source/StereoVideoIngestDevice/Private/StereoVideoIngestDevice.cpp b/Engine/Plugins/VirtualProduction/CaptureManager/CaptureManagerDevices/Source/StereoVideoIngestDevice/Private/StereoVideoIngestDevice.cpp
--- a/Engine/Plugins/VirtualProduction/CaptureManager/CaptureManagerDevices/Source/StereoVideoIngestDevice/Private/StereoVideoIngestDevice.cpp
+++ b/Engine/Plugins/VirtualProduction/CaptureManager/CaptureManagerDevices/Source/StereoVideoIngestDevice/Private/StereoVideoIngestDevice.cpp
@@ -122,21 +122,39 @@ FString UStereoVideoIngestDevice::GetFullTakePath(UE::CaptureManager::FTakeId In
 
 void UStereoVideoIngestDevice::RunConvertAndUploadTake(const UIngestCapability_ProcessHandle* InProcessHandle, const UIngestCapability_Options* InIngestOptions)
 {
-	AsyncTask(ENamedThreads::AnyThread, [this,
+	auto Task = [WeakThis = TWeakObjectPtr<UStereoVideoIngestDevice>(this),
 		ProcessHandle = TStrongObjectPtr<const UIngestCapability_ProcessHandle>(InProcessHandle),
 		IngestOptions = TStrongObjectPtr<const UIngestCapability_Options>(InIngestOptions)]()
 		{
+			TStrongObjectPtr<UStereoVideoIngestDevice> This = WeakThis.Pin();
+
+			if (!This.IsValid())
+			{
+				UE_LOG(LogVideoLiveLinkDevice, Warning, TEXT("Failed to ingest take, the device has been destroyed"));
+				return;
+			}
+
 			static constexpr uint32 NumberOfTasks = 2;
 
 			using namespace UE::CaptureManager;
 			TSharedPtr<FTaskProgress> TaskProgress
-				= MakeShared<FTaskProgress>(NumberOfTasks, FTaskProgress::FProgressReporter::CreateLambda([this, ProcessHandle](double InProgress)
+				= MakeShared<FTaskProgress>(NumberOfTasks, FTaskProgress::FProgressReporter::CreateLambda([WeakThis = WeakThis, ProcessHandle](double InProgress)
 					{
-						ExecuteProcessProgressReporter(ProcessHandle.Get(), InProgress);
+						TStrongObjectPtr<UStereoVideoIngestDevice> This = WeakThis.Pin();
+
+						if (!This.IsValid())
+						{
+							// Don't log this one, just because it's called so frequently
+							return;
+						}
+
+						This->ExecuteProcessProgressReporter(ProcessHandle.Get(), InProgress);
 					}));
 
-			Super::IngestTake(ProcessHandle.Get(), IngestOptions.Get(), MoveTemp(TaskProgress));
-		});
+			This->Super::IngestTake(ProcessHandle.Get(), IngestOptions.Get(), MoveTemp(TaskProgress));
+		};
+
+	UE::Tasks::Launch(UE_SOURCE_LOCATION, MoveTemp(Task), UE::Tasks::ETaskPriority::BackgroundNormal);
 }
 
 int32 UStereoVideoIngestDevice::FTakeWithComponents::CountComponents(UStereoVideoIngestDevice::ETakeComponentType Type)
@@ -154,46 +172,56 @@ int32 UStereoVideoIngestDevice::FTakeWithComponents::CountComponents(UStereoVide
 
 void UStereoVideoIngestDevice::UpdateTakeList_Implementation(UIngestCapability_UpdateTakeListCallback* InCallback)
 {
-	AsyncTask(ENamedThreads::AnyThread, [this, Callback = TStrongObjectPtr<UIngestCapability_UpdateTakeListCallback>(InCallback)]()
-	{
-		RemoveAllTakes();
+	auto Task = [WeakThis = TWeakObjectPtr<UStereoVideoIngestDevice>(this), Callback = TStrongObjectPtr<UIngestCapability_UpdateTakeListCallback>(InCallback)]()
+		{
+			TStrongObjectPtr<UStereoVideoIngestDevice> This = WeakThis.Pin();
 
-		FString StoragePath = GetSettings()->TakeDirectory.Path;
+			if (!This.IsValid())
+			{
+				UE_LOG(LogVideoLiveLinkDevice, Warning, TEXT("Failed to update take list, the device has been destroyed"));
+				return;
+			}
 
-		TMap<FString, UStereoVideoIngestDevice::FTakeWithComponents> TakesWithComponentsCandidates = DiscoverTakes(StoragePath);
+			This->RemoveAllTakes();
 
-		TArray<FString> SlateNames;
-		TakesWithComponentsCandidates.GetKeys(SlateNames);
+			FString StoragePath = This->GetSettings()->TakeDirectory.Path;
 
-		for (const FString& SlateName : SlateNames)
-		{
-			FTakeWithComponents& TakeWithComponents = TakesWithComponentsCandidates[SlateName];
+			TMap<FString, UStereoVideoIngestDevice::FTakeWithComponents> TakesWithComponentsCandidates = This->DiscoverTakes(StoragePath);
 
-			bool TakeAsExpected = ((TakeWithComponents.CountComponents(ETakeComponentType::VIDEO) == UE::StereoVideoLiveLinkDevice::Private::MaximumVideoFilesCountPerTake && TakeWithComponents.CountComponents(ETakeComponentType::IMAGE_SEQUENCE) == 0) ||
-				(TakeWithComponents.CountComponents(ETakeComponentType::VIDEO) == 0 && TakeWithComponents.CountComponents(ETakeComponentType::IMAGE_SEQUENCE) == UE::StereoVideoLiveLinkDevice::Private::MaximumVideoFilesCountPerTake)) &&
-				TakeWithComponents.CountComponents(ETakeComponentType::AUDIO) <= UE::StereoVideoLiveLinkDevice::Private::MaximumAudioFilesCountPerTake;
+			TArray<FString> SlateNames;
+			TakesWithComponentsCandidates.GetKeys(SlateNames);
 
-			if (!TakeAsExpected)
+			for (const FString& SlateName : SlateNames)
 			{
-				UE_LOG(LogVideoLiveLinkDevice, Warning, TEXT("Invalid take at '%s'. Take should have exactly two video files or two image sequences. Optionally, one audio file."), *TakeWithComponents.TakePath);
-				continue;
-			}
+				FTakeWithComponents& TakeWithComponents = TakesWithComponentsCandidates[SlateName];
 
-			TOptional<FTakeMetadata> TakeInfoResult = CreateTakeMetadata(TakeWithComponents);
+				bool TakeAsExpected = ((TakeWithComponents.CountComponents(ETakeComponentType::VIDEO) == UE::StereoVideoLiveLinkDevice::Private::MaximumVideoFilesCountPerTake && TakeWithComponents.CountComponents(ETakeComponentType::IMAGE_SEQUENCE) == 0) ||
+					(TakeWithComponents.CountComponents(ETakeComponentType::VIDEO) == 0 && TakeWithComponents.CountComponents(ETakeComponentType::IMAGE_SEQUENCE) == UE::StereoVideoLiveLinkDevice::Private::MaximumVideoFilesCountPerTake)) &&
+					TakeWithComponents.CountComponents(ETakeComponentType::AUDIO) <= UE::StereoVideoLiveLinkDevice::Private::MaximumAudioFilesCountPerTake;
 
-			if (TakeInfoResult.IsSet())
-			{
-				FTakeMetadata TakeMetadata = TakeInfoResult.GetValue();
+				if (!TakeAsExpected)
+				{
+					UE_LOG(LogVideoLiveLinkDevice, Warning, TEXT("Invalid take at '%s'. Take should have exactly two video files or two image sequences. Optionally, one audio file."), *TakeWithComponents.TakePath);
+					continue;
+				}
 
-				int32 TakeId = AddTake(MoveTemp(TakeMetadata));
-				FullTakePaths.Add(TakeId, TakeWithComponents.TakePath);
+				TOptional<FTakeMetadata> TakeInfoResult = This->CreateTakeMetadata(TakeWithComponents);
 
-				PublishEvent<FIngestCapability_TakeAddedEvent>(TakeId);
+				if (TakeInfoResult.IsSet())
+				{
+					FTakeMetadata TakeMetadata = TakeInfoResult.GetValue();
+
+					int32 TakeId = This->AddTake(MoveTemp(TakeMetadata));
+					This->FullTakePaths.Add(TakeId, TakeWithComponents.TakePath);
+
+					This->PublishEvent<FIngestCapability_TakeAddedEvent>(TakeId);
+				}
 			}
-		}
 
-		ExecuteUpdateTakeListCallback(Callback.Get(), Execute_GetTakeIdentifiers(this));
-	});
+			This->ExecuteUpdateTakeListCallback(Callback.Get(), Execute_GetTakeIdentifiers(This.Get()));
+		};
+
+	UE::Tasks::Launch(UE_SOURCE_LOCATION, MoveTemp(Task), UE::Tasks::ETaskPriority::BackgroundNormal);
 }
 
 TMap<FString, UStereoVideoIngestDevice::FTakeWithComponents> UStereoVideoIngestDevice::DiscoverTakes(FString StoragePath)
diff --git a/Engine/Plugins/VirtualProduction/CaptureManager/CaptureManagerDevices/Source/TakeArchiveIngestDevice/Private/TakeArchiveIngestDevice.cpp b/Engine/Plugins/VirtualProduction/CaptureManager/CaptureManagerDevices/Source/TakeArchiveIngestDevice/Private/TakeArchiveIngestDevice.cpp
--- a/Engine/Plugins/VirtualProduction/CaptureManager/CaptureManagerDevices/Source/TakeArchiveIngestDevice/Private/TakeArchiveIngestDevice.cpp
+++ b/Engine/Plugins/VirtualProduction/CaptureManager/CaptureManagerDevices/Source/TakeArchiveIngestDevice/Private/TakeArchiveIngestDevice.cpp
@@ -75,69 +75,97 @@ FString UTakeArchiveIngestDevice::GetFullTakePath(UE::CaptureManager::FTakeId In
 
 void UTakeArchiveIngestDevice::RunConvertAndUploadTake(const UIngestCapability_ProcessHandle* InProcessHandle, const UIngestCapability_Options* InIngestOptions)
 {
-	AsyncTask(ENamedThreads::AnyThread, [this,
-			  ProcessHandle = TStrongObjectPtr<const UIngestCapability_ProcessHandle>(InProcessHandle),
-			  IngestOptions = TStrongObjectPtr<const UIngestCapability_Options>(InIngestOptions)]()
-	{
-		static constexpr uint32 NumberOfTasks = 2; // Convert, Upload
-
-		using namespace UE::CaptureManager;
-		TSharedPtr<FTaskProgress> TaskProgress
-			= MakeShared<FTaskProgress>(NumberOfTasks, FTaskProgress::FProgressReporter::CreateLambda([this, ProcessHandle](double InProgress)
+	auto Task = [WeakThis = TWeakObjectPtr<UTakeArchiveIngestDevice>(this),
+		ProcessHandle = TStrongObjectPtr<const UIngestCapability_ProcessHandle>(InProcessHandle),
+		IngestOptions = TStrongObjectPtr<const UIngestCapability_Options>(InIngestOptions)]()
 		{
-			ExecuteProcessProgressReporter(ProcessHandle.Get(), InProgress);
-		}));
+			TStrongObjectPtr<UTakeArchiveIngestDevice> This = WeakThis.Pin();
+
+			if (!This.IsValid())
+			{
+				UE_LOG(LogTakeArchiveIngestDevice, Warning, TEXT("Failed to ingest take, the device has been destroyed"));
+				return;
+			}
+
+			static constexpr uint32 NumberOfTasks = 2; // Convert, Upload
+
+			using namespace UE::CaptureManager;
+			TSharedPtr<FTaskProgress> TaskProgress
+				= MakeShared<FTaskProgress>(NumberOfTasks, FTaskProgress::FProgressReporter::CreateLambda([WeakThis = WeakThis, ProcessHandle](double InProgress)
+					{
+						TStrongObjectPtr<UTakeArchiveIngestDevice> This = WeakThis.Pin();
 
-		Super::IngestTake(ProcessHandle.Get(), IngestOptions.Get(), MoveTemp(TaskProgress));
-	});
+						if (!This.IsValid())
+						{
+							// Don't log this one, just because it's called so frequently
+							return;
+						}
+
+						This->ExecuteProcessProgressReporter(ProcessHandle.Get(), InProgress);
+					}));
+
+			This->Super::IngestTake(ProcessHandle.Get(), IngestOptions.Get(), MoveTemp(TaskProgress));
+		};
+
+	UE::Tasks::Launch(UE_SOURCE_LOCATION, MoveTemp(Task), UE::Tasks::ETaskPriority::BackgroundNormal);
 }
 
 void UTakeArchiveIngestDevice::UpdateTakeList_Implementation(UIngestCapability_UpdateTakeListCallback* InCallback)
 {
-	AsyncTask(ENamedThreads::AnyThread, [this, Callback = TStrongObjectPtr<UIngestCapability_UpdateTakeListCallback>(InCallback)]()
-	{
-		RemoveAllTakes();
-
-		FString StoragePath = GetSettings()->TakeDirectory.Path;
-
-		TArray<FString> TakeMetadataFiles;
-		const bool bIterationResult = IFileManager::Get().IterateDirectoryRecursively(*StoragePath, [&](const TCHAR* InFileNameOrDirectory, bool bInIsDirectory)
+	auto Task = [WeakThis = TWeakObjectPtr<UTakeArchiveIngestDevice>(this), Callback = TStrongObjectPtr<UIngestCapability_UpdateTakeListCallback>(InCallback)]()
 		{
-			if (!bInIsDirectory)
+			TStrongObjectPtr<UTakeArchiveIngestDevice> This = WeakThis.Pin();
+
+			if (!This.IsValid())
 			{
-				const FString CurrentFileName = FPaths::GetCleanFilename(InFileNameOrDirectory);
-				if (CurrentFileName == TEXT("take.json"))
-				{
-					TakeMetadataFiles.Add(InFileNameOrDirectory);
-				}
-				else if (FPaths::GetExtension(CurrentFileName) == FTakeMetadata::FileExtension)
-				{
-					TakeMetadataFiles.Add(InFileNameOrDirectory);
-				}
+				UE_LOG(LogTakeArchiveIngestDevice, Warning, TEXT("Failed to update take list, the device has been destroyed"));
+				return;
 			}
 
-			return true;
-		});
+			This->RemoveAllTakes();
 
-		for (const FString& CurrentTakeFile : TakeMetadataFiles)
-		{
-			TOptional<FTakeMetadata> TakeInfoResult = ReadTake(CurrentTakeFile);
+			FString StoragePath = This->GetSettings()->TakeDirectory.Path;
 
-			if (TakeInfoResult.IsSet())
+			TArray<FString> TakeMetadataFiles;
+			const bool bIterationResult = IFileManager::Get().IterateDirectoryRecursively(*StoragePath, [&](const TCHAR* InFileNameOrDirectory, bool bInIsDirectory)
+				{
+					if (!bInIsDirectory)
+					{
+						const FString CurrentFileName = FPaths::GetCleanFilename(InFileNameOrDirectory);
+						if (CurrentFileName == TEXT("take.json"))
+						{
+							TakeMetadataFiles.Add(InFileNameOrDirectory);
+						}
+						else if (FPaths::GetExtension(CurrentFileName) == FTakeMetadata::FileExtension)
+						{
+							TakeMetadataFiles.Add(InFileNameOrDirectory);
+						}
+					}
+
+					return true;
+				});
+
+			for (const FString& CurrentTakeFile : TakeMetadataFiles)
 			{
-				int32 TakeId = AddTake(TakeInfoResult.GetValue());
-				FString CurrentDirectory = FPaths::GetPath(CurrentTakeFile);
+				TOptional<FTakeMetadata> TakeInfoResult = ReadTake(CurrentTakeFile);
+
+				if (TakeInfoResult.IsSet())
+				{
+					int32 TakeId = This->AddTake(TakeInfoResult.GetValue());
+					FString CurrentDirectory = FPaths::GetPath(CurrentTakeFile);
 
-				FString TakePath = CurrentDirectory.Right(CurrentDirectory.Len() - StoragePath.Len());
-				FPaths::NormalizeDirectoryName(TakePath); // Removes trailing slash
-				RelativeTakePaths.Add(TakeId, MoveTemp(TakePath));
+					FString TakePath = CurrentDirectory.Right(CurrentDirectory.Len() - StoragePath.Len());
+					FPaths::NormalizeDirectoryName(TakePath); // Removes trailing slash
+					This->RelativeTakePaths.Add(TakeId, MoveTemp(TakePath));
 
-				PublishEvent<FIngestCapability_TakeAddedEvent>(TakeId);
+					This->PublishEvent<FIngestCapability_TakeAddedEvent>(TakeId);
+				}
 			}
-		}
 
-		ExecuteUpdateTakeListCallback(Callback.Get(), Execute_GetTakeIdentifiers(this));
-	});
+			This->ExecuteUpdateTakeListCallback(Callback.Get(), Execute_GetTakeIdentifiers(This.Get()));
+		};
+
+	UE::Tasks::Launch(UE_SOURCE_LOCATION, MoveTemp(Task), UE::Tasks::ETaskPriority::BackgroundNormal);
 }
 
 ELiveLinkDeviceConnectionStatus UTakeArchiveIngestDevice::GetConnectionStatus_Implementation() const
