Fixed thumbnail image rotation when a third party encoder is configured in Capture Manager

#rnx
#Synced-CL 48514830

diff --git a/Engine/Plugins/VirtualProduction/CaptureManager/CaptureManagerApp/Source/IngestLiveLinkDevice/Private/Utils/CaptureExtractTimecode.cpp b/Engine/Plugins/VirtualProduction/CaptureManager/CaptureManagerApp/Source/IngestLiveLinkDevice/Private/Utils/CaptureExtractTimecode.cpp
--- a/Engine/Plugins/VirtualProduction/CaptureManager/CaptureManagerApp/Source/IngestLiveLinkDevice/Private/Utils/CaptureExtractTimecode.cpp
+++ b/Engine/Plugins/VirtualProduction/CaptureManager/CaptureManagerApp/Source/IngestLiveLinkDevice/Private/Utils/CaptureExtractTimecode.cpp
@@ -22,6 +22,25 @@
 
 DEFINE_LOG_CATEGORY_STATIC(LogCaptureExtractTimecode, Log, All);
 
+static bool IsEncoderFFmpeg(const FString& InEncoderPath)
+{
+	// Get base filename without extension
+	const FString EncoderFileName = FPaths::GetBaseFilename(InEncoderPath);
+	return EncoderFileName.Equals(TEXT("ffmpeg"), ESearchCase::IgnoreCase);
+}
+
+static FString GetFFprobePath(const FString& InEncoderPath)
+{
+	check(IsEncoderFFmpeg(InEncoderPath));
+
+	// Get base filename with extension: e.g ffmpeg.exe
+	const FString EncoderFileName = FPaths::GetCleanFilename(InEncoderPath);
+	const FString RootPath = FPaths::GetPath(InEncoderPath);
+	const FString FFprobeFileName = EncoderFileName.Replace(TEXT("ffmpeg"), TEXT("ffprobe"), ESearchCase::IgnoreCase);
+	const FString FFprobePath = RootPath / FFprobeFileName;
+	return FFprobePath;
+}
+
 FCaptureExtractVideoInfo::FResult FCaptureExtractVideoInfo::Create(const FString& InFilePath)
 {
 	using namespace UE::CaptureManager;
@@ -77,24 +96,16 @@ FCaptureExtractVideoInfo::FExtractResult FCaptureExtractVideoInfo::ExtractInfo(c
 	if (Settings->bEnableThirdPartyEncoder)
 	{
 		FString EncoderPath = Settings->ThirdPartyEncoder.FilePath;
-		FString BinarySuffix;
-#if PLATFORM_WINDOWS
-		BinarySuffix = TEXT(".exe"); // Temporary means of resolving ffprobe binary for 5.7.2. Refactor in 5.8 to use common code with VideoDeviceThumbnailExtractor.cpp
-#endif
-		if (EncoderPath.EndsWith(TEXT("ffmpeg") + BinarySuffix))
+		FString FFprobePath = GetFFprobePath(EncoderPath);
+		if (FPaths::FileExists(FFprobePath))
 		{
-			EncoderPath.ReplaceInline(*(TEXT("ffmpeg") + BinarySuffix), *(TEXT("ffprobe") + BinarySuffix));
-
-			if (FPaths::FileExists(EncoderPath))
-			{
-				FExtractResult Result = ExtractInfoUsingFFProbe(InFilePath, EncoderPath);
-				return Result;
-			}
-			else
-			{
-				UE_LOG(LogCaptureExtractTimecode, Warning, TEXT("Could not extract timecode for '%s' using ffprobe. ffprobe could not be found at path %s. ffprobe is expected to be found in the same directory as ffmpeg"), *InFilePath, *EncoderPath);
-				return MakeError(ECaptureExtractInfoError::FfprobeNotFound);
-			}
+			FExtractResult Result = ExtractInfoUsingFFProbe(InFilePath, FFprobePath);
+			return Result;
+		}
+		else
+		{
+			UE_LOG(LogCaptureExtractTimecode, Warning, TEXT("Could not extract timecode for '%s' using ffprobe. ffprobe could not be found at path %s. ffprobe is expected to be found in the same directory as ffmpeg"), *InFilePath, *EncoderPath);
+			return MakeError(ECaptureExtractInfoError::FfprobeNotFound);
 		}
 	}
 	else
diff --git a/Engine/Plugins/VirtualProduction/CaptureManager/CaptureManagerDevices/Source/VideoLiveLinkDeviceCommon/Private/Utils/VideoDeviceThumbnailExtractor.cpp b/Engine/Plugins/VirtualProduction/CaptureManager/CaptureManagerDevices/Source/VideoLiveLinkDeviceCommon/Private/Utils/VideoDeviceThumbnailExtractor.cpp
--- a/Engine/Plugins/VirtualProduction/CaptureManager/CaptureManagerDevices/Source/VideoLiveLinkDeviceCommon/Private/Utils/VideoDeviceThumbnailExtractor.cpp
+++ b/Engine/Plugins/VirtualProduction/CaptureManager/CaptureManagerDevices/Source/VideoLiveLinkDeviceCommon/Private/Utils/VideoDeviceThumbnailExtractor.cpp
@@ -6,6 +6,7 @@
 #include "Engine/Texture2D.h"
 #include "CaptureManagerMediaRWModule.h"
 #include "VideoLiveLinkDeviceLog.h"
+#include "Serialization/JsonSerializer.h"
 
 #include "Settings/CaptureManagerSettings.h"
 
@@ -16,30 +17,6 @@ DEFINE_LOG_CATEGORY(LogVideoLiveLinkDevice);
 namespace UE::CaptureManager
 {
 
-namespace Private
-{
-
-static bool IsEncoderFFmpeg(const FString& InEncoderPath)
-{
-	// Get base filename without extension
-	const FString EncoderFileName = FPaths::GetBaseFilename(InEncoderPath);
-	return EncoderFileName == TEXT("ffmpeg");
-}
-
-static FString GetFFprobePath(const FString& InEncoderPath) // Centralize this in 5.8 for use with CaptureExtractTimecode.cpp
-{
-	check(IsEncoderFFmpeg(InEncoderPath));
-
-	// Get base filename with extension: e.g ffmpeg.exe
-	const FString EncoderFileName = FPaths::GetCleanFilename(InEncoderPath);
-	const FString RootPath = FPaths::GetPath(InEncoderPath);
-	const FString FFprobeFileName = EncoderFileName.Replace(TEXT("ffmpeg"), TEXT("ffprobe"));
-	const FString FFprobePath = RootPath / FFprobeFileName;
-	return FFprobePath;
-}
-
-}
-
 FVideoDeviceThumbnailExtractor::FVideoDeviceThumbnailExtractor() = default;
 
 TOptional<FTakeThumbnailData::FRawImage> FVideoDeviceThumbnailExtractor::ExtractThumbnail(const FString& InCurrentFile)
@@ -50,61 +27,53 @@ TOptional<FTakeThumbnailData::FRawImage> FVideoDeviceThumbnailExtractor::Extract
 	int32 Width = 0;
 	int32 Height = 0;
 
-	TValueOrError<TUniquePtr<IVideoReader>, FText> VideoReaderResult = MediaManager.CreateVideoReader(InCurrentFile);
-	if (VideoReaderResult.HasValue())
+	// If the third party encoder is enabled, then always use ffmpeg to gather thumbnail data.
+	const UCaptureManagerSettings* Settings = GetDefault<UCaptureManagerSettings>();
+	if (Settings->bEnableThirdPartyEncoder)
 	{
-		TUniquePtr<IVideoReader> VideoReader = VideoReaderResult.StealValue();
-		TValueOrError<TUniquePtr<UE::CaptureManager::FMediaTextureSample>, FText> TextureSampleResult = VideoReader->Next();
-		if (TextureSampleResult.HasValue())
-		{
-			const UE::CaptureManager::FMediaTextureSample* Sample = TextureSampleResult.GetValue().Get();
+		Thumbnail = ObtainThumbnailFromThirdPartyEncoder(Settings->ThirdPartyEncoder.FilePath, InCurrentFile);
 
-			if (Sample)
-			{
-				Thumbnail = ConvertThumbnailFromSample(Sample);
-				Width = Sample->Dimensions.X;
-				Height = Sample->Dimensions.Y;
-			}
-		}
-		else
+		if (Thumbnail.IsEmpty())
 		{
-			FText ErrorText = TextureSampleResult.StealError();
-			UE_LOG(LogVideoLiveLinkDevice, Warning, TEXT("Couldn't obtain the thumbnail from the video %s : %s"), *InCurrentFile, *ErrorText.ToString());
+			return {};
 		}
+
+		// As we specified the width during image creation we can derive
+		// the height from the final number of pixels in the image.
+		Width = ThirdPartyEncoderThumbnailWidth;
+		Height = Thumbnail.Num() / Width;
 	}
-	else
+	else // Otherwise attempt to use an appropriate video reader for the file 
 	{
-		FText ErrorText = VideoReaderResult.StealError();
-		UE_LOG(LogVideoLiveLinkDevice, Warning, TEXT("Couldn't open video file %s : %s"), *InCurrentFile, *ErrorText.ToString());
-	}
-	
-	if (Thumbnail.IsEmpty())
-	{
-		const UCaptureManagerSettings* Settings = GetDefault<UCaptureManagerSettings>();
-		if (Settings->bEnableThirdPartyEncoder)
+		TValueOrError<TUniquePtr<IVideoReader>, FText> VideoReaderResult = MediaManager.CreateVideoReader(InCurrentFile);
+		if (VideoReaderResult.HasValue())
 		{
-			UE_LOG(LogVideoLiveLinkDevice, Warning, TEXT("Obtaining thumbnail using third party encoder"));
-
-			Thumbnail = ObtainThumbnailFromThirdPartyEncoder(Settings->ThirdPartyEncoder.FilePath, InCurrentFile);
-			TOptional<FIntPoint> DimensionsOpt = ObtainImageDimensionsFromThirdPartyEncoder(Settings->ThirdPartyEncoder.FilePath, InCurrentFile);
-
-			if (DimensionsOpt.IsSet())
+			TUniquePtr<IVideoReader> VideoReader = VideoReaderResult.StealValue();
+			TValueOrError<TUniquePtr<UE::CaptureManager::FMediaTextureSample>, FText> TextureSampleResult = VideoReader->Next();
+			if (TextureSampleResult.HasValue())
 			{
-				FIntPoint Dimensions = MoveTemp(DimensionsOpt.GetValue());
-
-				Width = ThirdPartyEncoderThumbnailWidth;
-				Height = FMath::RoundToInt((static_cast<float>(Width) * Dimensions.Y) / Dimensions.X);
-
-				check(Thumbnail.Num() == Width * Height);
+				const UE::CaptureManager::FMediaTextureSample* Sample = TextureSampleResult.GetValue().Get();
+
+				if (Sample)
+				{
+					Thumbnail = ConvertThumbnailFromSample(Sample);
+					Width = Sample->Dimensions.X;
+					Height = Sample->Dimensions.Y;
+				}
 			}
 			else
 			{
-				UE_LOG(LogVideoLiveLinkDevice, Warning, TEXT("Couldn't obtain frame dimensions of the video file %s"), *InCurrentFile);
-				return {};
+				FText ErrorText = TextureSampleResult.StealError();
+				UE_LOG(LogVideoLiveLinkDevice, Warning, TEXT("Couldn't obtain the thumbnail from the video %s : %s"), *InCurrentFile, *ErrorText.ToString());
 			}
 		}
+		else
+		{
+			FText ErrorText = VideoReaderResult.StealError();
+			UE_LOG(LogVideoLiveLinkDevice, Warning, TEXT("Couldn't open video file %s : %s"), *InCurrentFile, *ErrorText.ToString());
+		}
 	}
-
+	
 	if (Thumbnail.IsEmpty())
 	{
 		return {};
@@ -207,77 +176,26 @@ TArray<FColor> FVideoDeviceThumbnailExtractor::ObtainThumbnailFromThirdPartyEnco
 	FString ScaleArg = FString::Printf(TEXT("scale=%d:-1"), ThirdPartyEncoderThumbnailWidth);
 
 	FString CommandLineArgs = FString::Format(TEXT("-hide_banner -loglevel error -i \"{0}\" -map v:0 -vf \"thumbnail=2, {1}\" -frames:v 1 -f rawvideo -pix_fmt rgb24 -"), { InCurrentFile, ScaleArg });
-	
-	TArray<uint8> ThumbnailRawData = RunProcess(InEncoderPath, CommandLineArgs);
-
-	TArray<FColor> ThumbnailRawColorData;
-	ThumbnailRawColorData.Reserve(ThumbnailRawData.Num() / 3);
-
-	for (int32 ValueIndex = 0; ValueIndex < ThumbnailRawData.Num(); ValueIndex += 3)
-	{
-		ThumbnailRawColorData.Add(FColor(ThumbnailRawData[ValueIndex], ThumbnailRawData[ValueIndex + 1], ThumbnailRawData[ValueIndex + 2]));
-	}
-
-	return ThumbnailRawColorData;
-}
-
-TOptional<FIntPoint> FVideoDeviceThumbnailExtractor::ObtainImageDimensionsFromThirdPartyEncoder(const FString& InEncoderPath, const FString& InCurrentFile)
-{
-	if (!Private::IsEncoderFFmpeg(InEncoderPath))
-	{
-		return {};
-	}
-
-	FString CommandLine = FString::Format(TEXT("-v error -select_streams v:0 -show_entries stream=width,height -of default=noprint_wrappers=1 \"{0}\""), { InCurrentFile });
-	FString EncoderPath = Private::GetFFprobePath(InEncoderPath);
 
-	if (!FPaths::FileExists(EncoderPath))
-	{
-		UE_LOG(LogVideoLiveLinkDevice, Warning, TEXT("Could not extract image dimenstions for '%s' using ffprobe. ffprobe could not be found at path %s. ffprobe is expected to be found in the same directory as ffmpeg"), *InCurrentFile, *InEncoderPath);
-		return {};
-	}
-
-	TArray<uint8> WidthAndHeight = RunProcess(EncoderPath, CommandLine);
+	TArray<uint8> ThumbnailRawData = RunProcess(InEncoderPath, CommandLineArgs);
+	int32 ThumbnailNumBytes = ThumbnailRawData.Num();
 
-	if (WidthAndHeight.IsEmpty())
+	constexpr int32 NumChannels = 3;
+	if (ThumbnailNumBytes % NumChannels != 0)
 	{
-		return {};
+		UE_LOG(LogVideoLiveLinkDevice, Warning, TEXT("Thumbnail raw data is malformed and the byte count is not a multiple of %d. Num bytes is %d"), NumChannels, ThumbnailNumBytes);
+		return TArray<FColor>();
 	}
 
-	FString WidthAndHeightString = 
-		FString::ConstructFromPtrSize(reinterpret_cast<const UTF8CHAR*>(WidthAndHeight.GetData()), WidthAndHeight.Num());
-
-	TArray<FString> Lines;
-	WidthAndHeightString.ParseIntoArrayLines(Lines);
-
-	FIntPoint Dimensions = FIntPoint::ZeroValue;
-	for (const FString& Line : Lines)
-	{
-		TArray<FString> KeyValue;
-		Line.ParseIntoArray(KeyValue, TEXT("="));
-
-		if (KeyValue.Num() >= 2)
-		{
-			const FString& Key = KeyValue[0];
-			const FString& Value = KeyValue[1];
-
-			if (Key.Equals(TEXT("width"), ESearchCase::IgnoreCase))
-			{
-				Dimensions.X = FCString::Atoi(*Value);
-			}
-			else if (Key.Equals(TEXT("height"), ESearchCase::IgnoreCase))
-			{
-				Dimensions.Y = FCString::Atoi(*Value);
-			}
-		}
-	}
+	TArray<FColor> ThumbnailRawColorData;
+	ThumbnailRawColorData.Reserve(ThumbnailNumBytes / NumChannels);
 
-	if (Dimensions == FIntPoint::ZeroValue)
+	for (int32 ValueIndex = 0; ValueIndex < ThumbnailNumBytes; ValueIndex += NumChannels)
 	{
-		return {};
+		ThumbnailRawColorData.Add(FColor(ThumbnailRawData[ValueIndex], ThumbnailRawData[ValueIndex + 1], ThumbnailRawData[ValueIndex + 2]));
 	}
 
-	return Dimensions;
+	return ThumbnailRawColorData;
 }
 
 TArray<uint8> FVideoDeviceThumbnailExtractor::RunProcess(const FString& InProcessName, const FString& InProcessArgs)
@@ -325,7 +243,13 @@ TArray<uint8> FVideoDeviceThumbnailExtractor::RunProcess(const FString& InProces
 		return TArray<uint8>();
 	}
 
-	FullCommandOutput.Append(ReadFromPipe(ReadPipe));
+	// Drain pipe of all remaining data
+	TArray<uint8> CommandOutput;
+	do {
+		CommandOutput = ReadFromPipe(ReadPipe);
+		FullCommandOutput.Append(CommandOutput);
+	} while (!CommandOutput.IsEmpty());
+
 	return FullCommandOutput;
 }
 
diff --git a/Engine/Plugins/VirtualProduction/CaptureManager/CaptureManagerDevices/Source/VideoLiveLinkDeviceCommon/Public/Utils/VideoDeviceThumbnailExtractor.h b/Engine/Plugins/VirtualProduction/CaptureManager/CaptureManagerDevices/Source/VideoLiveLinkDeviceCommon/Public/Utils/VideoDeviceThumbnailExtractor.h
--- a/Engine/Plugins/VirtualProduction/CaptureManager/CaptureManagerDevices/Source/VideoLiveLinkDeviceCommon/Public/Utils/VideoDeviceThumbnailExtractor.h
+++ b/Engine/Plugins/VirtualProduction/CaptureManager/CaptureManagerDevices/Source/VideoLiveLinkDeviceCommon/Public/Utils/VideoDeviceThumbnailExtractor.h
@@ -28,7 +28,6 @@ private:
 	static constexpr int32 ThirdPartyEncoderThumbnailWidth = 256;
 
 	TArray<FColor> ObtainThumbnailFromThirdPartyEncoder(const FString& InEncoderPath, const FString& InCurrentFile);
-	TOptional<FIntPoint> ObtainImageDimensionsFromThirdPartyEncoder(const FString& InEncoderPath, const FString& InCurrentFile);
 
 	TArray<uint8> RunProcess(const FString& InProcessName, const FString& InProcessArgs);
 	TArray<uint8> ReadFromPipe(void* InPipe);
