[CaptureManager] Removed redundant timecode extraction calls and set timecode and frame rate correctly for audio in Mono Video Ingest device

#Synced-CL 49035297

diff --git a/Engine/Plugins/VirtualProduction/CaptureManager/CaptureManagerDevices/Source/MonoVideoIngestDevice/Private/MonoVideoIngestDevice.cpp b/Engine/Plugins/VirtualProduction/CaptureManager/CaptureManagerDevices/Source/MonoVideoIngestDevice/Private/MonoVideoIngestDevice.cpp
--- a/Engine/Plugins/VirtualProduction/CaptureManager/CaptureManagerDevices/Source/MonoVideoIngestDevice/Private/MonoVideoIngestDevice.cpp
+++ b/Engine/Plugins/VirtualProduction/CaptureManager/CaptureManagerDevices/Source/MonoVideoIngestDevice/Private/MonoVideoIngestDevice.cpp
@@ -305,13 +305,6 @@ TOptional<FTakeMetadata> UMonoVideoIngestDevice::ReadTake(const FString& InCurre
 	TakeMetadata.Device.Model = TEXT("MonoVideo");
 	
 	FTakeMetadata::FVideo Video;
-	TValueOrError<TUniquePtr<IVideoReader>, FText> VideoReaderResult = MediaManager.CreateVideoReader(InCurrentTakeFile);
-	if (VideoReaderResult.HasValue())
-	{
-		TUniquePtr<IVideoReader> VideoReader = VideoReaderResult.StealValue();
-		Video.FrameRate = static_cast<float>(VideoReader->GetFrameRate().AsDecimal());
-	}
-
 	FVideoDeviceThumbnailExtractor ThumbnailExtractor;
 	TOptional<FTakeThumbnailData::FRawImage> RawImageOpt = ThumbnailExtractor.ExtractThumbnail(InCurrentTakeFile);
 
@@ -330,10 +323,7 @@ TOptional<FTakeMetadata> UMonoVideoIngestDevice::ReadTake(const FString& InCurre
 	{
 		FCaptureExtractVideoInfo Extractor = ExtractorOpt.StealValue();
 		Video.TimecodeStart = Extractor.GetTimecode().ToString();
-		if (FMath::IsNearlyZero(Video.FrameRate))
-		{
-			Video.FrameRate = static_cast<float>(Extractor.GetFrameRate().AsDecimal());
-		}
+		Video.FrameRate = static_cast<float>(Extractor.GetFrameRate().AsDecimal());
 	}
 
 	if (!Video.TimecodeStart.IsSet())
@@ -346,7 +336,7 @@ TOptional<FTakeMetadata> UMonoVideoIngestDevice::ReadTake(const FString& InCurre
 		UE_LOG(LogVideoLiveLinkDevice, Warning, TEXT("Failed to determine frame rate for video file: %s"), *InCurrentTakeFile);
 	}
 
-	TakeMetadata.Video.Add(MoveTemp(Video));
+	TakeMetadata.Video.Add(Video);
 
 	// Check if the video file contains audio stream by creating a audio reader from the file
 	// If creation fails, we assume that there is no audio in the file
@@ -359,6 +349,9 @@ TOptional<FTakeMetadata> UMonoVideoIngestDevice::ReadTake(const FString& InCurre
 		Audio.Path = InCurrentTakeFile;
 		Audio.Duration = static_cast<float>(AudioReaderResult.StealValue()->GetDuration().GetTotalSeconds());
 
+		Audio.TimecodeStart = Video.TimecodeStart;
+		Audio.TimecodeRate = static_cast<float>(Video.FrameRate);
+
 		TakeMetadata.Audio.Add(MoveTemp(Audio));
 	}
 
diff --git a/Engine/Plugins/VirtualProduction/CaptureManager/CaptureManagerDevices/Source/StereoVideoIngestDevice/Private/StereoVideoIngestDevice.cpp b/Engine/Plugins/VirtualProduction/CaptureManager/CaptureManagerDevices/Source/StereoVideoIngestDevice/Private/StereoVideoIngestDevice.cpp
--- a/Engine/Plugins/VirtualProduction/CaptureManager/CaptureManagerDevices/Source/StereoVideoIngestDevice/Private/StereoVideoIngestDevice.cpp
+++ b/Engine/Plugins/VirtualProduction/CaptureManager/CaptureManagerDevices/Source/StereoVideoIngestDevice/Private/StereoVideoIngestDevice.cpp
@@ -580,19 +580,6 @@ TOptional<FTakeMetadata> UStereoVideoIngestDevice::CreateTakeMetadata(const FTak
 			Video.Format = FPaths::GetExtension(VideoComponent.Path);
 			Video.PathType = FTakeMetadata::FVideo::EPathType::File;
 
-			TValueOrError<TUniquePtr<IVideoReader>, FText> VideoReaderResult = MediaManager.CreateVideoReader(VideoComponent.Path);
-			if (VideoReaderResult.HasValue())
-			{
-				TUniquePtr<IVideoReader> VideoReader = VideoReaderResult.StealValue();
-				Video.FrameRate = static_cast<float>(VideoReader->GetFrameRate().AsDecimal());
-
-				if (!bVideoFrameRateSet)
-				{
-					VideoFrameRate = UE::CaptureManager::ParseFrameRate(Video.FrameRate);
-					bVideoFrameRateSet = true;
-				}
-			}
-
 			if (!TakeMetadata.Thumbnail.GetRawImage().IsSet())
 			{
 				FVideoDeviceThumbnailExtractor ThumbnailExtractor;
@@ -612,13 +599,10 @@ TOptional<FTakeMetadata> UStereoVideoIngestDevice::CreateTakeMetadata(const FTak
 			if (ExtractorOpt.IsValid())
 			{
 				FCaptureExtractVideoInfo Extractor = ExtractorOpt.StealValue();
-				Video.TimecodeStart = Extractor.GetTimecode().ToString();
-				if (FMath::IsNearlyZero(Video.FrameRate))
-				{
-					Video.FrameRate = static_cast<float>(Extractor.GetFrameRate().AsDecimal());
-					VideoFrameRate = Extractor.GetFrameRate();
-					bVideoFrameRateSet = true;
-				}
+				Video.TimecodeStart = Extractor.GetTimecode().ToString();			
+				Video.FrameRate = static_cast<float>(Extractor.GetFrameRate().AsDecimal());
+				VideoFrameRate = Extractor.GetFrameRate();
+				bVideoFrameRateSet = true;
 			}
 
 			if (!Video.TimecodeStart.IsSet())
@@ -636,40 +620,44 @@ TOptional<FTakeMetadata> UStereoVideoIngestDevice::CreateTakeMetadata(const FTak
 			Video.PathType = FTakeMetadata::FVideo::EPathType::Folder;
 
 			int32 FramesCount = 0;
-			bool FirstFrame = true;
+			bool bFirstFrame = true;
 
 			bool bIterationResult = IFileManager::Get().IterateDirectoryRecursively(*VideoComponent.Path, [&](const TCHAR* InPath, bool bInIsDirectory)
 				{
-					if (FirstFrame)
+					if (!bInIsDirectory)
 					{
-						FirstFrame = false;
-
-						if (!TakeMetadata.Thumbnail.GetThumbnailData().IsSet())
+						if (bFirstFrame)
 						{
-							TArray<uint8> ThumbnailData;
-							if (FFileHelper::LoadFileToArray(ThumbnailData, InPath))
+							bFirstFrame = false;
+
+							if (!TakeMetadata.Thumbnail.GetThumbnailData().IsSet())
 							{
-								TakeMetadata.Thumbnail = ThumbnailData;
+								TArray<uint8> ThumbnailData;
+								if (FFileHelper::LoadFileToArray(ThumbnailData, InPath))
+								{
+									TakeMetadata.Thumbnail = ThumbnailData;
+								}
 							}
-						}
 
-						FString Path, Filename, Extension;
-						FPaths::Split(InPath, Path, Filename, Extension);
-						Video.Format = Extension;
+							FString Path, Filename, Extension;
+							FPaths::Split(InPath, Path, Filename, Extension);
+							Video.Format = Extension;
 
-						Video.Path = VideoComponent.Path;
-						{
-							FImage Image;
-							FImageUtils::LoadImage(InPath, Image);
-							if (Image.IsImageInfoValid())
+							Video.Path = VideoComponent.Path;
 							{
-								Video.FrameWidth = Image.SizeX;
-								Video.FrameHeight = Image.SizeY;
+								FImage Image;
+								FImageUtils::LoadImage(InPath, Image);
+								if (Image.IsImageInfoValid())
+								{
+									Video.FrameWidth = Image.SizeX;
+									Video.FrameHeight = Image.SizeY;
+								}
 							}
 						}
-					}
 
-					FramesCount++;
+						// Assume found files are image files and increase frame count
+						FramesCount++;
+					}
 
 					return true;
 				});
