[CaptureManager] Crash on close when updating take list

#Synced-CL 49368151

diff --git a/Engine/Plugins/VirtualProduction/CaptureManager/CaptureManagerApp/Source/CaptureManagerEditor/Private/UI/CaptureManagerPanelController.cpp b/Engine/Plugins/VirtualProduction/CaptureManager/CaptureManagerApp/Source/CaptureManagerEditor/Private/UI/CaptureManagerPanelController.cpp
--- a/Engine/Plugins/VirtualProduction/CaptureManager/CaptureManagerApp/Source/CaptureManagerEditor/Private/UI/CaptureManagerPanelController.cpp
+++ b/Engine/Plugins/VirtualProduction/CaptureManager/CaptureManagerApp/Source/CaptureManagerEditor/Private/UI/CaptureManagerPanelController.cpp
@@ -59,6 +59,8 @@ FCaptureManagerPanelController::FCaptureManagerPanelController(FPrivateToken InP
 	Subsystem->OnDeviceAdded().AddRaw(this, &FCaptureManagerPanelController::OnDeviceAdded);
 	Subsystem->OnDeviceRemoved().AddRaw(this, &FCaptureManagerPanelController::OnDeviceRemoved);
 
+	UpdateTakeListManager = FUpdateTakeListManager::Create();
+
 	UnrealEndpointManager->Start();
 }
 
@@ -71,6 +73,8 @@ FCaptureManagerPanelController::~FCaptureManagerPanelController()
 	IngestJobProcessorWidget->OnProcessingStateChanged().Unbind();
 	IngestJobProcessorWidget->OnSelectionChanged().Unbind();
 
+	UpdateTakeListManager = nullptr;
+
 	UnrealEndpointManager->Stop();
 }
 
@@ -192,6 +196,11 @@ TArray<TObjectPtr<ULiveLinkDevice>> FCaptureManagerPanelController::GetCaptureDe
 	return CaptureDevices;
 }
 
+TSharedRef<FUpdateTakeListManager> FCaptureManagerPanelController::GetUpdateTakeListManager() const
+{
+	return UpdateTakeListManager.ToSharedRef();
+}
+
 void FCaptureManagerPanelController::OnDeviceAdded(FGuid InDeviceId, ULiveLinkDevice* InDevice)
 {
 	if (!InDevice->Implements<ULiveLinkDeviceCapability_Ingest>())
diff --git a/Engine/Plugins/VirtualProduction/CaptureManager/CaptureManagerApp/Source/CaptureManagerEditor/Private/UI/CaptureManagerPanelController.h b/Engine/Plugins/VirtualProduction/CaptureManager/CaptureManagerApp/Source/CaptureManagerEditor/Private/UI/CaptureManagerPanelController.h
--- a/Engine/Plugins/VirtualProduction/CaptureManager/CaptureManagerApp/Source/CaptureManagerEditor/Private/UI/CaptureManagerPanelController.h
+++ b/Engine/Plugins/VirtualProduction/CaptureManager/CaptureManagerApp/Source/CaptureManagerEditor/Private/UI/CaptureManagerPanelController.h
@@ -9,6 +9,8 @@
 #include "IngestManagement/IngestJobProcessor.h"
 #include "CaptureManagerUnrealEndpointManager.h"
 
+#include "UpdateTakeListManager.h"
+
 #include "LiveLinkDeviceSubsystem.h"
 #include "LiveLinkDevice.h"
 #include "LiveLinkDeviceCapability_Connection.h"
@@ -55,6 +57,8 @@ public:
 	TObjectPtr<ULiveLinkDevice> GetCaptureDevice(FGuid InDeviceId) const;
 	TArray<TObjectPtr<ULiveLinkDevice>> GetCaptureDevices() const;
 
+	TSharedRef<FUpdateTakeListManager> GetUpdateTakeListManager() const;
+
 private:
 
 	void OnDeviceAdded(FGuid InDeviceId, ULiveLinkDevice* InDevice);
@@ -82,4 +86,6 @@ private:
 	TSharedPtr<UE::CaptureManager::SIngestJobProcessor> IngestJobProcessorWidget;
 	TSharedPtr<IDetailsView> IngestJobDetailsView;
 	FDelegateHandle FinishedChangingPropertiesHandle;
+
+	TSharedPtr<FUpdateTakeListManager> UpdateTakeListManager;
 };
\ No newline at end of file
diff --git a/Engine/Plugins/VirtualProduction/CaptureManager/CaptureManagerApp/Source/CaptureManagerEditor/Private/UI/TakesView.cpp b/Engine/Plugins/VirtualProduction/CaptureManager/CaptureManagerApp/Source/CaptureManagerEditor/Private/UI/TakesView.cpp
--- a/Engine/Plugins/VirtualProduction/CaptureManager/CaptureManagerApp/Source/CaptureManagerEditor/Private/UI/TakesView.cpp
+++ b/Engine/Plugins/VirtualProduction/CaptureManager/CaptureManagerApp/Source/CaptureManagerEditor/Private/UI/TakesView.cpp
@@ -479,29 +479,35 @@ void FTakesView::Refresh()
 
 	if (SelectedItem.IsValid() && SelectedItem->Implements<ULiveLinkDeviceCapability_Ingest>())
 	{
-		RefreshSingleTake(SelectedItem.Get());
+		RefreshSingleDevice(SelectedItem.Get());
 	}
 	else
 	{
-		RefreshAllTakes();
+		RefreshAllDevices();
 	}
 }
 
-void FTakesView::RefreshSingleTake(TObjectPtr<ULiveLinkDevice> InDevice)
+void FTakesView::RefreshSingleDevice(TObjectPtr<ULiveLinkDevice> InDevice)
 {
 	ResetTakesCache(InDevice->GetDeviceId());
 
-	FIngestUpdateTakeListCallback Callback =
-		FIngestUpdateTakeListCallback::Type::CreateRaw(this, &FTakesView::UpdateTakeListCallback, InDevice->GetDeviceId());
+	TSharedPtr<FCaptureManagerPanelController> Controller = PanelController.Pin();
 
-	UIngestCapability_UpdateTakeListCallback* UpdateTakeListCallback = NewObject<UIngestCapability_UpdateTakeListCallback>();
-	UpdateTakeListCallback->Callback = MoveTemp(Callback);
+	if (!Controller)
+	{
+		return;
+	}
 
 	TakesTileView->UpdateTakeListStarted();
-	ILiveLinkDeviceCapability_Ingest::Execute_UpdateTakeList(InDevice, UpdateTakeListCallback);
+
+	FIngestUpdateTakeListCallback Callback =
+		FIngestUpdateTakeListCallback::Type::CreateRaw(this, &FTakesView::UpdateTakeListCallback, InDevice->GetDeviceId());
+
+	TSharedRef<FUpdateTakeListManager> UpdateTakeListManager = Controller->GetUpdateTakeListManager();
+	UpdateTakeListManager->UpdateTakeListForDevice(InDevice->GetDeviceId(), MoveTemp(Callback));
 }
 
-void FTakesView::RefreshAllTakes()
+void FTakesView::RefreshAllDevices()
 {
 	TSharedPtr<FCaptureManagerPanelController> Controller = PanelController.Pin();
 
@@ -512,7 +518,7 @@ void FTakesView::RefreshAllTakes()
 
 	for (const TObjectPtr<ULiveLinkDevice>& Device : Controller->GetCaptureDevices())
 	{
-		RefreshSingleTake(Device);
+		RefreshSingleDevice(Device);
 	}
 }
 
@@ -645,17 +651,7 @@ void FTakesView::CaptureDeviceStarted(FGuid InCaptureDeviceId)
 				return;
 			}
 
-			ResetTakesCache(InCaptureDeviceId);
-
-			FIngestUpdateTakeListCallback Callback = 
-				FIngestUpdateTakeListCallback::Type::CreateRaw(this, &FTakesView::UpdateTakeListCallback, InCaptureDeviceId);
-
-			UIngestCapability_UpdateTakeListCallback* UpdateTakeListCallback = NewObject<UIngestCapability_UpdateTakeListCallback>();
-			UpdateTakeListCallback->Callback = MoveTemp(Callback);
-
-			TakesTileView->UpdateTakeListStarted();
-
-			ILiveLinkDeviceCapability_Ingest::Execute_UpdateTakeList(Device, UpdateTakeListCallback);
+			RefreshSingleDevice(Device);
 		});
 }
 
diff --git a/Engine/Plugins/VirtualProduction/CaptureManager/CaptureManagerApp/Source/CaptureManagerEditor/Private/UI/TakesView.h b/Engine/Plugins/VirtualProduction/CaptureManager/CaptureManagerApp/Source/CaptureManagerEditor/Private/UI/TakesView.h
--- a/Engine/Plugins/VirtualProduction/CaptureManager/CaptureManagerApp/Source/CaptureManagerEditor/Private/UI/TakesView.h
+++ b/Engine/Plugins/VirtualProduction/CaptureManager/CaptureManagerApp/Source/CaptureManagerEditor/Private/UI/TakesView.h
@@ -104,8 +104,8 @@ private:
 	EImageRotation ConvertOrientation(TOptional<FTakeMetadata::FVideo::EOrientation> InOrientation) const;
 	void AddTakesToIngestQueue(const TArray<TObjectPtr<UTakeVirtualAsset>>& InTakeAssets, const UIngestJobSettings& InJobPreset);
 	void Refresh();
-	void RefreshSingleTake(TObjectPtr<ULiveLinkDevice> InDevice);
-	void RefreshAllTakes();
+	void RefreshSingleDevice(TObjectPtr<ULiveLinkDevice> InDevice);
+	void RefreshAllDevices();
 	void UpdateTakeListCallback(TArray<UE::CaptureManager::FTakeId> InTakeIds, FGuid InCaptureDeviceId);
 	void ResetTakesCache(FGuid InCaptureDeviceId);
 
diff --git a/Engine/Plugins/VirtualProduction/CaptureManager/CaptureManagerApp/Source/CaptureManagerEditor/Private/UI/UpdateTakeListManager.cpp b/Engine/Plugins/VirtualProduction/CaptureManager/CaptureManagerApp/Source/CaptureManagerEditor/Private/UI/UpdateTakeListManager.cpp
new file mode 100644
--- /dev/null
+++ b/Engine/Plugins/VirtualProduction/CaptureManager/CaptureManagerApp/Source/CaptureManagerEditor/Private/UI/UpdateTakeListManager.cpp
@@ -0,0 +1,81 @@
+// Copyright Epic Games, Inc. All Rights Reserved.
+
+#include "UpdateTakeListManager.h"
+#include "LiveLinkDeviceSubsystem.h"
+#include "LiveLinkDevice.h"
+
+#include "Ingest/LiveLinkDeviceCapability_Ingest.h"
+
+#include "Engine/Engine.h"
+
+TSharedPtr<FUpdateTakeListManager> FUpdateTakeListManager::Create()
+{
+	TSharedPtr<FUpdateTakeListManager> Manager = 
+		MakeShared<FUpdateTakeListManager>(FPrivate());
+
+	return Manager;
+}
+
+FUpdateTakeListManager::FUpdateTakeListManager(FPrivate)
+	: QueueRunner(UE::CaptureManager::TQueueRunner<FQueueRunnerContext>::FOnProcess::CreateStatic(&FUpdateTakeListManager::UpdateTakeListForDevice_Private))
+{
+}
+
+FUpdateTakeListManager::~FUpdateTakeListManager()
+{
+	QueueRunner.Empty();
+}
+
+
+void FUpdateTakeListManager::UpdateTakeListForDevice(FGuid InDeviceId, FIngestUpdateTakeListCallback InCallback)
+{
+	FQueueRunnerContext Context;
+	Context.DeviceId = MoveTemp(InDeviceId);
+	Context.Callback = MoveTemp(InCallback);
+
+	QueueRunner.Add(MoveTemp(Context));
+}
+
+void FUpdateTakeListManager::UpdateTakeListForDevice_Private(FQueueRunnerContext InContext)
+{
+
+	const FGuid& DeviceId = InContext.DeviceId;
+
+	if (!DeviceId.IsValid())
+	{
+		return;
+	}
+
+	if (!GEngine)
+	{
+		return;
+	}
+
+	ULiveLinkDeviceSubsystem* Subsystem = GEngine->GetEngineSubsystem<ULiveLinkDeviceSubsystem>();
+	if (!Subsystem)
+	{
+		return;
+	}
+
+	const TObjectPtr<ULiveLinkDevice>* DevicePtr = Subsystem->GetDeviceMap().Find(DeviceId);
+	if (!DevicePtr)
+	{
+		return;
+	}
+
+	TObjectPtr<ULiveLinkDevice> Device = *DevicePtr;
+	if (!Device)
+	{
+		return;
+	}
+
+	if (!Device->Implements<ULiveLinkDeviceCapability_Ingest>())
+	{
+		return;
+	}
+
+	UIngestCapability_UpdateTakeListCallback* UpdateTakeListCallback = NewObject<UIngestCapability_UpdateTakeListCallback>();
+	UpdateTakeListCallback->Callback = MoveTemp(InContext.Callback);
+
+	ILiveLinkDeviceCapability_Ingest::Execute_UpdateTakeList(Device, UpdateTakeListCallback);
+}
diff --git a/Engine/Plugins/VirtualProduction/CaptureManager/CaptureManagerApp/Source/CaptureManagerEditor/Private/UI/UpdateTakeListManager.h b/Engine/Plugins/VirtualProduction/CaptureManager/CaptureManagerApp/Source/CaptureManagerEditor/Private/UI/UpdateTakeListManager.h
new file mode 100644
--- /dev/null
+++ b/Engine/Plugins/VirtualProduction/CaptureManager/CaptureManagerApp/Source/CaptureManagerEditor/Private/UI/UpdateTakeListManager.h
@@ -0,0 +1,39 @@
+// Copyright Epic Games, Inc. All Rights Reserved.
+
+#pragma once
+
+#include "Async/QueueRunner.h"
+#include "Templates/SharedPointer.h"
+#include "Misc/Guid.h"
+
+#include "LiveLinkDevice.h"
+#include "Ingest/LiveLinkDeviceCapability_Ingest.h"
+#include "Ingest/IngestCapability_UpdateTakeList.h"
+
+class FUpdateTakeListManager
+{
+private:
+
+	struct FPrivate { explicit FPrivate() = default; };
+
+public:
+
+	static TSharedPtr<FUpdateTakeListManager> Create();
+
+	FUpdateTakeListManager(FPrivate InToken);
+	~FUpdateTakeListManager();
+
+	void UpdateTakeListForDevice(FGuid InDeviceId, FIngestUpdateTakeListCallback InCallback);
+
+private:
+
+	struct FQueueRunnerContext
+	{
+		FGuid DeviceId;
+		FIngestUpdateTakeListCallback Callback;
+	};
+
+	static void UpdateTakeListForDevice_Private(FQueueRunnerContext InContext);
+
+	UE::CaptureManager::TQueueRunner<FQueueRunnerContext> QueueRunner;
+};
\ No newline at end of file
diff --git a/Engine/Plugins/VirtualProduction/CaptureManager/CaptureManagerApp/Source/IngestLiveLinkDevice/Private/BaseIngestLiveLinkDevice.cpp b/Engine/Plugins/VirtualProduction/CaptureManager/CaptureManagerApp/Source/IngestLiveLinkDevice/Private/BaseIngestLiveLinkDevice.cpp
--- a/Engine/Plugins/VirtualProduction/CaptureManager/CaptureManagerApp/Source/IngestLiveLinkDevice/Private/BaseIngestLiveLinkDevice.cpp
+++ b/Engine/Plugins/VirtualProduction/CaptureManager/CaptureManagerApp/Source/IngestLiveLinkDevice/Private/BaseIngestLiveLinkDevice.cpp
@@ -430,6 +430,8 @@ void UBaseIngestLiveLinkDevice::OnDeviceAdded()
 
 void UBaseIngestLiveLinkDevice::OnDeviceRemoved()
 {
+	Execute_RequestAbortUpdateTakeList(this);
+
 	Impl->OnDeviceRemoved(Execute_GetTakeIdentifiers(this));
 
 	RemoveAllTakes();
diff --git a/Engine/Plugins/VirtualProduction/CaptureManager/CaptureManagerApp/Source/IngestLiveLinkDevice/Public/BaseIngestLiveLinkDevice.h b/Engine/Plugins/VirtualProduction/CaptureManager/CaptureManagerApp/Source/IngestLiveLinkDevice/Public/BaseIngestLiveLinkDevice.h
--- a/Engine/Plugins/VirtualProduction/CaptureManager/CaptureManagerApp/Source/IngestLiveLinkDevice/Public/BaseIngestLiveLinkDevice.h
+++ b/Engine/Plugins/VirtualProduction/CaptureManager/CaptureManagerApp/Source/IngestLiveLinkDevice/Public/BaseIngestLiveLinkDevice.h
@@ -73,6 +73,13 @@ protected:
 	 */
 	virtual void RunConvertAndUploadTake(const UIngestCapability_ProcessHandle* InProcessHandle, const UIngestCapability_Options* InIngestOptions) override PURE_VIRTUAL(UBaseIngestLiveLinkDevice::RunConvertAndUploadTake);
 
+	/**
+	 * Inherited function from the ILiveLinkDeviceCapability_Ingest. The function is required to be implemented.
+	 * @note Runs update take list in the current thread
+	 * @param Callback to be run
+	 */
+	virtual void RunUpdateTakeList(UIngestCapability_UpdateTakeListCallback* InCallback) override PURE_VIRTUAL(UBaseIngestLiveLinkDevice::RunUpdateTakeList);
+
 	//~ End ULiveLinkDeviceCapability_Ingest interface
 
 	/**
diff --git a/Engine/Plugins/VirtualProduction/CaptureManager/CaptureManagerApp/Source/LiveLinkCapabilities/Private/Ingest/LiveLinkDeviceCapability_Ingest.cpp b/Engine/Plugins/VirtualProduction/CaptureManager/CaptureManagerApp/Source/LiveLinkCapabilities/Private/Ingest/LiveLinkDeviceCapability_Ingest.cpp
--- a/Engine/Plugins/VirtualProduction/CaptureManager/CaptureManagerApp/Source/LiveLinkCapabilities/Private/Ingest/LiveLinkDeviceCapability_Ingest.cpp
+++ b/Engine/Plugins/VirtualProduction/CaptureManager/CaptureManagerApp/Source/LiveLinkCapabilities/Private/Ingest/LiveLinkDeviceCapability_Ingest.cpp
@@ -82,6 +82,18 @@ void ILiveLinkDeviceCapability_Ingest::RunIngestProcess_Implementation(UIngestCa
 	RunIngestProcess(InProcessHandle);
 }
 
+void ILiveLinkDeviceCapability_Ingest::UpdateTakeList_Implementation(UIngestCapability_UpdateTakeListCallback* InCallback)
+{
+	StopToken = StopRequester.CreateToken();
+
+	RunUpdateTakeList(InCallback);
+}
+
+void ILiveLinkDeviceCapability_Ingest::RequestAbortUpdateTakeList_Implementation()
+{
+	StopRequester.RequestStop();
+}
+
 UIngestCapability_TakeInformation* ILiveLinkDeviceCapability_Ingest::GetTakeInformation_Implementation(int32 InTakeId) const
 {
 	TOptional<FTakeMetadata> TakeMetadataOpt = GetTakeMetadata(InTakeId);
@@ -117,6 +129,17 @@ TArray<int32> ILiveLinkDeviceCapability_Ingest::GetTakeIdentifiers_Implementatio
 	return TakeIdentifiers;
 }
 
+TOptional<FTakeMetadata> ILiveLinkDeviceCapability_Ingest::GetTakeMetadata(int32 InTakeId) const
+{
+	FScopeLock Lock(&TakeMetadataMutex);
+	if (const FTakeMetadata* TakeMetadata = TakeMetadataMap.Find(InTakeId))
+	{
+		return *TakeMetadata;
+	}
+
+	return {};
+}
+
 void ILiveLinkDeviceCapability_Ingest::ExecuteProcessFinishedReporter(const UIngestCapability_ProcessHandle* InProcessHandle, TValueOrError<void, FIngestCapability_Error> InMaybeError)
 {
 	UIngestCapability_ProcessResult* Result =
@@ -212,15 +235,9 @@ bool ILiveLinkDeviceCapability_Ingest::UpdateTake(int32 InTakeId, FTakeMetadata
 	return false;
 }
 
-TOptional<FTakeMetadata> ILiveLinkDeviceCapability_Ingest::GetTakeMetadata(int32 InTakeId) const
+bool ILiveLinkDeviceCapability_Ingest::IsUpdateTakeListAbortRequested() const
 {
-	FScopeLock Lock(&TakeMetadataMutex);
-	if (const FTakeMetadata* TakeMetadata = TakeMetadataMap.Find(InTakeId))
-	{
-		return *TakeMetadata;
-	}
-
-	return {};
+	return StopToken.IsStopRequested();
 }
 
 void ILiveLinkDeviceCapability_Ingest::RunIngestProcess(const UIngestCapability_ProcessHandle* InProcessHandle)
diff --git a/Engine/Plugins/VirtualProduction/CaptureManager/CaptureManagerApp/Source/LiveLinkCapabilities/Public/Ingest/LiveLinkDeviceCapability_Ingest.h b/Engine/Plugins/VirtualProduction/CaptureManager/CaptureManagerApp/Source/LiveLinkCapabilities/Public/Ingest/LiveLinkDeviceCapability_Ingest.h
--- a/Engine/Plugins/VirtualProduction/CaptureManager/CaptureManagerApp/Source/LiveLinkCapabilities/Public/Ingest/LiveLinkDeviceCapability_Ingest.h
+++ b/Engine/Plugins/VirtualProduction/CaptureManager/CaptureManagerApp/Source/LiveLinkCapabilities/Public/Ingest/LiveLinkDeviceCapability_Ingest.h
@@ -7,6 +7,8 @@
 #include "Templates/ValueOrError.h"
 #include "CaptureManagerTakeMetadata.h"
 
+#include "Async/StopToken.h"
+
 #include "Ingest/IngestCapability_Options.h"
 #include "Ingest/IngestCapability_ProcessHandle.h"
 #include "Ingest/IngestCapability_UpdateTakeList.h"
@@ -57,6 +59,9 @@ public:
 	UFUNCTION(BlueprintCallable, BlueprintNativeEvent, Category = "Live Link Device|Ingest")
 	void UpdateTakeList(UIngestCapability_UpdateTakeListCallback* InCallback);
 
+	UFUNCTION(BlueprintCallable, BlueprintNativeEvent, Category = "Live Link Device|Ingest")
+	void RequestAbortUpdateTakeList();
+
 	UFUNCTION(BlueprintCallable, BlueprintNativeEvent, Category = "Live Link Device|Ingest")
 	UIngestCapability_TakeInformation* GetTakeInformation(int32 InTakeId) const;
 
@@ -67,7 +72,9 @@ public:
 	void RunIngestProcess_Implementation(UIngestCapability_ProcessHandle* InProcessHandle, const UIngestCapability_Options* InOptions);
 	virtual void CancelIngestProcess_Implementation(const UIngestCapability_ProcessHandle* InProcessHandle) PURE_VIRTUAL(ILiveLinkDeviceCapability_Ingest::CancelIngestProcess_Implementation);
 
-	virtual void UpdateTakeList_Implementation(UIngestCapability_UpdateTakeListCallback* InCallback) PURE_VIRTUAL(ILiveLinkDeviceCapability_Ingest::UpdateTakeList_Implementation);
+	void UpdateTakeList_Implementation(UIngestCapability_UpdateTakeListCallback* InCallback);
+	void RequestAbortUpdateTakeList_Implementation();
+
 	UIngestCapability_TakeInformation* GetTakeInformation_Implementation(int32 InTakeId) const;
 	TArray<int32> GetTakeIdentifiers_Implementation() const;
 
@@ -83,6 +90,7 @@ protected:
 	void RemoveTake(int32 InTakeId);
 	void RemoveAllTakes();
 	bool UpdateTake(int32 InTakeId, FTakeMetadata InTakeMetadata);
+	bool IsUpdateTakeListAbortRequested() const;
 
 private:
 
@@ -92,6 +100,10 @@ private:
 
 	virtual void RunDownloadTake(const UIngestCapability_ProcessHandle* InProcessHandle, const UIngestCapability_Options* InIngestOptions) = 0;
 	virtual void RunConvertAndUploadTake(const UIngestCapability_ProcessHandle* InProcessHandle, const UIngestCapability_Options* InIngestOptions) = 0;
+	virtual void RunUpdateTakeList(UIngestCapability_UpdateTakeListCallback* InCallback) = 0;
+
+	UE::CaptureManager::FStopRequester StopRequester;
+	UE::CaptureManager::FStopToken StopToken;
 
 	mutable FCriticalSection TakeMetadataMutex;
 	TMap<int32, FTakeMetadata> TakeMetadataMap;
diff --git a/Engine/Plugins/VirtualProduction/CaptureManager/CaptureManagerCore/Source/CaptureUtils/Public/Async/QueueRunner.h b/Engine/Plugins/VirtualProduction/CaptureManager/CaptureManagerCore/Source/CaptureUtils/Public/Async/QueueRunner.h
--- a/Engine/Plugins/VirtualProduction/CaptureManager/CaptureManagerCore/Source/CaptureUtils/Public/Async/QueueRunner.h
+++ b/Engine/Plugins/VirtualProduction/CaptureManager/CaptureManagerCore/Source/CaptureUtils/Public/Async/QueueRunner.h
@@ -20,9 +20,9 @@ public:
 
 	DECLARE_DELEGATE_OneParam(FOnProcess, T InElem)
 
-	TQueueRunner(FOnProcess InProcess)
+	TQueueRunner(FOnProcess InOnProcess = nullptr)
 		: bRunning(false)
-		, OnProcess(MoveTemp(InProcess))
+		, OnProcess(MoveTemp(InOnProcess))
 	{
 		Thread.Reset(FRunnableThread::Create(this, TEXT("Queue Runner"), 128 * 1024, TPri_Normal, FPlatformAffinity::GetPoolThreadMask()));
 	}
diff --git a/Engine/Plugins/VirtualProduction/CaptureManager/CaptureManagerDevices/Source/CPSLiveLinkDevice/Private/LiveLinkFaceDevice.cpp b/Engine/Plugins/VirtualProduction/CaptureManager/CaptureManagerDevices/Source/CPSLiveLinkDevice/Private/LiveLinkFaceDevice.cpp
--- a/Engine/Plugins/VirtualProduction/CaptureManager/CaptureManagerDevices/Source/CPSLiveLinkDevice/Private/LiveLinkFaceDevice.cpp
+++ b/Engine/Plugins/VirtualProduction/CaptureManager/CaptureManagerDevices/Source/CPSLiveLinkDevice/Private/LiveLinkFaceDevice.cpp
@@ -76,7 +76,7 @@ FString ULiveLinkFaceDevice::GetFullTakePath(UE::CaptureManager::FTakeId InTakeI
 	return FString();
 }
 
-void ULiveLinkFaceDevice::UpdateTakeList_Implementation(UIngestCapability_UpdateTakeListCallback* InCallback)
+void ULiveLinkFaceDevice::RunUpdateTakeList(UIngestCapability_UpdateTakeListCallback* InCallback)
 {
 	using namespace UE::CaptureManager;
 
diff --git a/Engine/Plugins/VirtualProduction/CaptureManager/CaptureManagerDevices/Source/CPSLiveLinkDevice/Public/LiveLinkFaceDevice.h b/Engine/Plugins/VirtualProduction/CaptureManager/CaptureManagerDevices/Source/CPSLiveLinkDevice/Public/LiveLinkFaceDevice.h
--- a/Engine/Plugins/VirtualProduction/CaptureManager/CaptureManagerDevices/Source/CPSLiveLinkDevice/Public/LiveLinkFaceDevice.h
+++ b/Engine/Plugins/VirtualProduction/CaptureManager/CaptureManagerDevices/Source/CPSLiveLinkDevice/Public/LiveLinkFaceDevice.h
@@ -65,7 +65,7 @@ private:
 	virtual FString GetFullTakePath(UE::CaptureManager::FTakeId InTakeId) const override;
 
 	//~ Begin ULiveLinkDeviceCapability_Ingest interface
-	virtual void UpdateTakeList_Implementation(UIngestCapability_UpdateTakeListCallback* InCallback) override;
+	virtual void RunUpdateTakeList(UIngestCapability_UpdateTakeListCallback* InCallback) override;
 	virtual void RunDownloadTake(const UIngestCapability_ProcessHandle* InProcessHandle, const UIngestCapability_Options* InIngestOptions) override;
 	virtual void RunConvertAndUploadTake(const UIngestCapability_ProcessHandle* InProcessHandle, const UIngestCapability_Options* InIngestOptions) override;
 	virtual void CancelIngestProcess_Implementation(const UIngestCapability_ProcessHandle* InProcessHandle) override;
diff --git a/Engine/Plugins/VirtualProduction/CaptureManager/CaptureManagerDevices/Source/MonoVideoIngestDevice/Private/MonoVideoIngestDevice.cpp b/Engine/Plugins/VirtualProduction/CaptureManager/CaptureManagerDevices/Source/MonoVideoIngestDevice/Private/MonoVideoIngestDevice.cpp
--- a/Engine/Plugins/VirtualProduction/CaptureManager/CaptureManagerDevices/Source/MonoVideoIngestDevice/Private/MonoVideoIngestDevice.cpp
+++ b/Engine/Plugins/VirtualProduction/CaptureManager/CaptureManagerDevices/Source/MonoVideoIngestDevice/Private/MonoVideoIngestDevice.cpp
@@ -138,61 +138,58 @@ void UMonoVideoIngestDevice::RunConvertAndUploadTake(const UIngestCapability_Pro
 	UE::Tasks::Launch(UE_SOURCE_LOCATION, MoveTemp(Task), UE::Tasks::ETaskPriority::BackgroundNormal);
 }
 
-void UMonoVideoIngestDevice::UpdateTakeList_Implementation(UIngestCapability_UpdateTakeListCallback* InCallback)
+void UMonoVideoIngestDevice::RunUpdateTakeList(UIngestCapability_UpdateTakeListCallback* InCallback)
 {
-	auto Task = [WeakThis = TWeakObjectPtr<UMonoVideoIngestDevice>(this), Callback = TStrongObjectPtr<UIngestCapability_UpdateTakeListCallback>(InCallback)]
-		{
-			TStrongObjectPtr<UMonoVideoIngestDevice> This = WeakThis.Pin();
+	RemoveAllTakes();
 
-			if (!This.IsValid())
+	FString StoragePath = GetSettings()->TakeDirectory.Path;
+	int32 DirectoriesInterrogatedCount = 0;
+	const int32 DirectoriesToInterrogateInOneRun = 200;
+
+	TArray<FString> SupportedVideoFiles;
+	const bool bIterationResult = IFileManager::Get().IterateDirectoryRecursively(*StoragePath, [&](const TCHAR* InFileNameOrDirectory, bool bInIsDirectory)
+		{
+			if (IsUpdateTakeListAbortRequested())
 			{
-				UE_LOG(LogVideoLiveLinkDevice, Warning, TEXT("Failed to update take list, the device has been destroyed"));
-				return;
+				return false;
 			}
 
-			This->RemoveAllTakes();
-
-			FString StoragePath = This->GetSettings()->TakeDirectory.Path;
-			int32 DirectoriesInterrogatedCount = 0;
-			const int32 DirectoriesToInterrogateInOneRun = 200;
-
-			TArray<FString> SupportedVideoFiles;
-			const bool bIterationResult = IFileManager::Get().IterateDirectoryRecursively(*StoragePath, [&](const TCHAR* InFileNameOrDirectory, bool bInIsDirectory)
+			if (bInIsDirectory)
+			{
+				if (++DirectoriesInterrogatedCount > DirectoriesToInterrogateInOneRun)
 				{
-					if (bInIsDirectory)
-					{
-						if (++DirectoriesInterrogatedCount > DirectoriesToInterrogateInOneRun)
-						{
-							return false;
-						}
-					}
-
-					if (!bInIsDirectory && This->IsVideoFile(InFileNameOrDirectory))
-					{
-						SupportedVideoFiles.Add(InFileNameOrDirectory);
-					}
-
-					return true;
-				});
+					return false;
+				}
+			}
 
-			for (const FString& SupportedVideoFile : SupportedVideoFiles)
+			if (!bInIsDirectory && IsVideoFile(InFileNameOrDirectory))
 			{
-				TOptional<FTakeMetadata> TakeInfoResult = This->ReadTake(SupportedVideoFile);
+				SupportedVideoFiles.Add(InFileNameOrDirectory);
+			}
 
-				if (TakeInfoResult.IsSet())
-				{
-					int32 TakeId = This->AddTake(TakeInfoResult.GetValue());
-					This->FullTakePaths.Add(TakeId, SupportedVideoFile);
+			return true;
+		});
 
-					This->PublishEvent<FIngestCapability_TakeAddedEvent>(TakeId);
-				}
-			}
+	for (const FString& SupportedVideoFile : SupportedVideoFiles)
+	{
+		if (IsUpdateTakeListAbortRequested())
+		{
+			ExecuteUpdateTakeListCallback(InCallback, TArray<int32>());
+			return;
+		}
 
-			This->ExecuteUpdateTakeListCallback(Callback.Get(), Execute_GetTakeIdentifiers(This.Get()));
-		};
+		TOptional<FTakeMetadata> TakeInfoResult = ReadTake(SupportedVideoFile);
 
+		if (TakeInfoResult.IsSet())
+		{
+			int32 TakeId = AddTake(TakeInfoResult.GetValue());
+			FullTakePaths.Add(TakeId, SupportedVideoFile);
 
-	UE::Tasks::Launch(UE_SOURCE_LOCATION, MoveTemp(Task), UE::Tasks::ETaskPriority::BackgroundNormal);
+			PublishEvent<FIngestCapability_TakeAddedEvent>(TakeId);
+		}
+	}
+
+	ExecuteUpdateTakeListCallback(InCallback, Execute_GetTakeIdentifiers(this));
 }
 
 ELiveLinkDeviceConnectionStatus UMonoVideoIngestDevice::GetConnectionStatus_Implementation() const
diff --git a/Engine/Plugins/VirtualProduction/CaptureManager/CaptureManagerDevices/Source/MonoVideoIngestDevice/Private/MonoVideoIngestDevice.h b/Engine/Plugins/VirtualProduction/CaptureManager/CaptureManagerDevices/Source/MonoVideoIngestDevice/Private/MonoVideoIngestDevice.h
--- a/Engine/Plugins/VirtualProduction/CaptureManager/CaptureManagerDevices/Source/MonoVideoIngestDevice/Private/MonoVideoIngestDevice.h
+++ b/Engine/Plugins/VirtualProduction/CaptureManager/CaptureManagerDevices/Source/MonoVideoIngestDevice/Private/MonoVideoIngestDevice.h
@@ -86,7 +86,7 @@ private:
 
 	//~ Begin ULiveLinkDeviceCapability_Ingest interface
 	virtual void RunConvertAndUploadTake(const UIngestCapability_ProcessHandle* InProcessHandle, const UIngestCapability_Options* InIngestOptions) override;
-	virtual void UpdateTakeList_Implementation(UIngestCapability_UpdateTakeListCallback* InCallback) override;
+	virtual void RunUpdateTakeList(UIngestCapability_UpdateTakeListCallback* InCallback) override;
 	//~ End ULiveLinkDeviceCapability_Ingest interface
 
 	//~ Begin ILiveLinkDeviceCapability_Connection interface
diff --git a/Engine/Plugins/VirtualProduction/CaptureManager/CaptureManagerDevices/Source/StereoVideoIngestDevice/Private/StereoVideoIngestDevice.cpp b/Engine/Plugins/VirtualProduction/CaptureManager/CaptureManagerDevices/Source/StereoVideoIngestDevice/Private/StereoVideoIngestDevice.cpp
--- a/Engine/Plugins/VirtualProduction/CaptureManager/CaptureManagerDevices/Source/StereoVideoIngestDevice/Private/StereoVideoIngestDevice.cpp
+++ b/Engine/Plugins/VirtualProduction/CaptureManager/CaptureManagerDevices/Source/StereoVideoIngestDevice/Private/StereoVideoIngestDevice.cpp
@@ -170,58 +170,51 @@ int32 UStereoVideoIngestDevice::FTakeWithComponents::CountComponents(UStereoVide
 	return Count;
 }
 
-void UStereoVideoIngestDevice::UpdateTakeList_Implementation(UIngestCapability_UpdateTakeListCallback* InCallback)
+void UStereoVideoIngestDevice::RunUpdateTakeList(UIngestCapability_UpdateTakeListCallback* InCallback)
 {
-	auto Task = [WeakThis = TWeakObjectPtr<UStereoVideoIngestDevice>(this), Callback = TStrongObjectPtr<UIngestCapability_UpdateTakeListCallback>(InCallback)]()
-		{
-			TStrongObjectPtr<UStereoVideoIngestDevice> This = WeakThis.Pin();
-
-			if (!This.IsValid())
-			{
-				UE_LOG(LogVideoLiveLinkDevice, Warning, TEXT("Failed to update take list, the device has been destroyed"));
-				return;
-			}
+	RemoveAllTakes();
 
-			This->RemoveAllTakes();
+	FString StoragePath = GetSettings()->TakeDirectory.Path;
 
-			FString StoragePath = This->GetSettings()->TakeDirectory.Path;
+	TMap<FString, UStereoVideoIngestDevice::FTakeWithComponents> TakesWithComponentsCandidates = DiscoverTakes(StoragePath);
 
-			TMap<FString, UStereoVideoIngestDevice::FTakeWithComponents> TakesWithComponentsCandidates = This->DiscoverTakes(StoragePath);
+	TArray<FString> SlateNames;
+	TakesWithComponentsCandidates.GetKeys(SlateNames);
 
-			TArray<FString> SlateNames;
-			TakesWithComponentsCandidates.GetKeys(SlateNames);
-
-			for (const FString& SlateName : SlateNames)
-			{
-				FTakeWithComponents& TakeWithComponents = TakesWithComponentsCandidates[SlateName];
+	for (const FString& SlateName : SlateNames)
+	{
+		if (IsUpdateTakeListAbortRequested())
+		{
+			ExecuteUpdateTakeListCallback(InCallback, TArray<int32>());
+			return;
+		}
 
-				bool TakeAsExpected = ((TakeWithComponents.CountComponents(ETakeComponentType::VIDEO) == UE::StereoVideoLiveLinkDevice::Private::MaximumVideoFilesCountPerTake && TakeWithComponents.CountComponents(ETakeComponentType::IMAGE_SEQUENCE) == 0) ||
-					(TakeWithComponents.CountComponents(ETakeComponentType::VIDEO) == 0 && TakeWithComponents.CountComponents(ETakeComponentType::IMAGE_SEQUENCE) == UE::StereoVideoLiveLinkDevice::Private::MaximumVideoFilesCountPerTake)) &&
-					TakeWithComponents.CountComponents(ETakeComponentType::AUDIO) <= UE::StereoVideoLiveLinkDevice::Private::MaximumAudioFilesCountPerTake;
+		FTakeWithComponents& TakeWithComponents = TakesWithComponentsCandidates[SlateName];
 
-				if (!TakeAsExpected)
-				{
-					UE_LOG(LogVideoLiveLinkDevice, Warning, TEXT("Invalid take at '%s'. Take should have exactly two video files or two image sequences. Optionally, one audio file."), *TakeWithComponents.TakePath);
-					continue;
-				}
+		bool TakeAsExpected = ((TakeWithComponents.CountComponents(ETakeComponentType::VIDEO) == UE::StereoVideoLiveLinkDevice::Private::MaximumVideoFilesCountPerTake && TakeWithComponents.CountComponents(ETakeComponentType::IMAGE_SEQUENCE) == 0) ||
+							   (TakeWithComponents.CountComponents(ETakeComponentType::VIDEO) == 0 && TakeWithComponents.CountComponents(ETakeComponentType::IMAGE_SEQUENCE) == UE::StereoVideoLiveLinkDevice::Private::MaximumVideoFilesCountPerTake)) &&
+			TakeWithComponents.CountComponents(ETakeComponentType::AUDIO) <= UE::StereoVideoLiveLinkDevice::Private::MaximumAudioFilesCountPerTake;
 
-				TOptional<FTakeMetadata> TakeInfoResult = This->CreateTakeMetadata(TakeWithComponents);
+		if (!TakeAsExpected)
+		{
+			UE_LOG(LogVideoLiveLinkDevice, Warning, TEXT("Invalid take at '%s'. Take should have exactly two video files or two image sequences. Optionally, one audio file."), *TakeWithComponents.TakePath);
+			continue;
+		}
 
-				if (TakeInfoResult.IsSet())
-				{
-					FTakeMetadata TakeMetadata = TakeInfoResult.GetValue();
+		TOptional<FTakeMetadata> TakeInfoResult = CreateTakeMetadata(TakeWithComponents);
 
-					int32 TakeId = This->AddTake(MoveTemp(TakeMetadata));
-					This->FullTakePaths.Add(TakeId, TakeWithComponents.TakePath);
+		if (TakeInfoResult.IsSet())
+		{
+			FTakeMetadata TakeMetadata = TakeInfoResult.GetValue();
 
-					This->PublishEvent<FIngestCapability_TakeAddedEvent>(TakeId);
-				}
-			}
+			int32 TakeId = AddTake(MoveTemp(TakeMetadata));
+			FullTakePaths.Add(TakeId, TakeWithComponents.TakePath);
 
-			This->ExecuteUpdateTakeListCallback(Callback.Get(), Execute_GetTakeIdentifiers(This.Get()));
-		};
+			PublishEvent<FIngestCapability_TakeAddedEvent>(TakeId);
+		}
+	}
 
-	UE::Tasks::Launch(UE_SOURCE_LOCATION, MoveTemp(Task), UE::Tasks::ETaskPriority::BackgroundNormal);
+	ExecuteUpdateTakeListCallback(InCallback, Execute_GetTakeIdentifiers(this));
 }
 
 TMap<FString, UStereoVideoIngestDevice::FTakeWithComponents> UStereoVideoIngestDevice::DiscoverTakes(FString StoragePath)
@@ -232,6 +225,11 @@ TMap<FString, UStereoVideoIngestDevice::FTakeWithComponents> UStereoVideoIngestD
 
 	const bool bIterationResult = IFileManager::Get().IterateDirectoryRecursively(*StoragePath, [&](const TCHAR* InPath, bool bInIsDirectory)
 		{
+			if (IsUpdateTakeListAbortRequested())
+			{
+				return false;
+			}
+
 			if (bInIsDirectory)
 			{
 				if (++DirectoriesInterrogatedCount > UE::StereoVideoLiveLinkDevice::Private::DirectoriesCountToIterateForTakesSearch)
diff --git a/Engine/Plugins/VirtualProduction/CaptureManager/CaptureManagerDevices/Source/StereoVideoIngestDevice/Private/StereoVideoIngestDevice.h b/Engine/Plugins/VirtualProduction/CaptureManager/CaptureManagerDevices/Source/StereoVideoIngestDevice/Private/StereoVideoIngestDevice.h
--- a/Engine/Plugins/VirtualProduction/CaptureManager/CaptureManagerDevices/Source/StereoVideoIngestDevice/Private/StereoVideoIngestDevice.h
+++ b/Engine/Plugins/VirtualProduction/CaptureManager/CaptureManagerDevices/Source/StereoVideoIngestDevice/Private/StereoVideoIngestDevice.h
@@ -143,7 +143,7 @@ private:
 	//~ Begin ULiveLinkDeviceCapability_Ingest interface
 	virtual void RunConvertAndUploadTake(const UIngestCapability_ProcessHandle* InProcessHandle,
 										 const UIngestCapability_Options* InIngestOptions) override;
-	virtual void UpdateTakeList_Implementation(UIngestCapability_UpdateTakeListCallback* InCallback) override;
+	virtual void RunUpdateTakeList(UIngestCapability_UpdateTakeListCallback* InCallback) override;
 	//~ End ULiveLinkDeviceCapability_Ingest interface
 
 	//~ Begin ILiveLinkDeviceCapability_Connection interface
diff --git a/Engine/Plugins/VirtualProduction/CaptureManager/CaptureManagerDevices/Source/TakeArchiveIngestDevice/Private/TakeArchiveIngestDevice.cpp b/Engine/Plugins/VirtualProduction/CaptureManager/CaptureManagerDevices/Source/TakeArchiveIngestDevice/Private/TakeArchiveIngestDevice.cpp
--- a/Engine/Plugins/VirtualProduction/CaptureManager/CaptureManagerDevices/Source/TakeArchiveIngestDevice/Private/TakeArchiveIngestDevice.cpp
+++ b/Engine/Plugins/VirtualProduction/CaptureManager/CaptureManagerDevices/Source/TakeArchiveIngestDevice/Private/TakeArchiveIngestDevice.cpp
@@ -110,62 +110,60 @@ void UTakeArchiveIngestDevice::RunConvertAndUploadTake(const UIngestCapability_P
 	UE::Tasks::Launch(UE_SOURCE_LOCATION, MoveTemp(Task), UE::Tasks::ETaskPriority::BackgroundNormal);
 }
 
-void UTakeArchiveIngestDevice::UpdateTakeList_Implementation(UIngestCapability_UpdateTakeListCallback* InCallback)
+void UTakeArchiveIngestDevice::RunUpdateTakeList(UIngestCapability_UpdateTakeListCallback* InCallback)
 {
-	auto Task = [WeakThis = TWeakObjectPtr<UTakeArchiveIngestDevice>(this), Callback = TStrongObjectPtr<UIngestCapability_UpdateTakeListCallback>(InCallback)]()
-		{
-			TStrongObjectPtr<UTakeArchiveIngestDevice> This = WeakThis.Pin();
+	RemoveAllTakes();
 
-			if (!This.IsValid())
+	FString StoragePath = GetSettings()->TakeDirectory.Path;
+
+	TArray<FString> TakeMetadataFiles;
+	const bool bIterationResult = IFileManager::Get().IterateDirectoryRecursively(*StoragePath, [&](const TCHAR* InFileNameOrDirectory, bool bInIsDirectory)
+		{
+			if (IsUpdateTakeListAbortRequested())
 			{
-				UE_LOG(LogTakeArchiveIngestDevice, Warning, TEXT("Failed to update take list, the device has been destroyed"));
-				return;
+				return false;
 			}
 
-			This->RemoveAllTakes();
-
-			FString StoragePath = This->GetSettings()->TakeDirectory.Path;
-
-			TArray<FString> TakeMetadataFiles;
-			const bool bIterationResult = IFileManager::Get().IterateDirectoryRecursively(*StoragePath, [&](const TCHAR* InFileNameOrDirectory, bool bInIsDirectory)
+			if (!bInIsDirectory)
+			{
+				const FString CurrentFileName = FPaths::GetCleanFilename(InFileNameOrDirectory);
+				if (CurrentFileName == TEXT("take.json"))
 				{
-					if (!bInIsDirectory)
-					{
-						const FString CurrentFileName = FPaths::GetCleanFilename(InFileNameOrDirectory);
-						if (CurrentFileName == TEXT("take.json"))
-						{
-							TakeMetadataFiles.Add(InFileNameOrDirectory);
-						}
-						else if (FPaths::GetExtension(CurrentFileName) == FTakeMetadata::FileExtension)
-						{
-							TakeMetadataFiles.Add(InFileNameOrDirectory);
-						}
-					}
+					TakeMetadataFiles.Add(InFileNameOrDirectory);
+				}
+				else if (FPaths::GetExtension(CurrentFileName) == FTakeMetadata::FileExtension)
+				{
+					TakeMetadataFiles.Add(InFileNameOrDirectory);
+				}
+			}
 
-					return true;
-				});
+			return true;
+		});
 
-			for (const FString& CurrentTakeFile : TakeMetadataFiles)
-			{
-				TOptional<FTakeMetadata> TakeInfoResult = ReadTake(CurrentTakeFile);
+	for (const FString& CurrentTakeFile : TakeMetadataFiles)
+	{
+		if (IsUpdateTakeListAbortRequested())
+		{
+			ExecuteUpdateTakeListCallback(InCallback, TArray<int32>());
+			return;
+		}
 
-				if (TakeInfoResult.IsSet())
-				{
-					int32 TakeId = This->AddTake(TakeInfoResult.GetValue());
-					FString CurrentDirectory = FPaths::GetPath(CurrentTakeFile);
+		TOptional<FTakeMetadata> TakeInfoResult = ReadTake(CurrentTakeFile);
 
-					FString TakePath = CurrentDirectory.Right(CurrentDirectory.Len() - StoragePath.Len());
-					FPaths::NormalizeDirectoryName(TakePath); // Removes trailing slash
-					This->RelativeTakePaths.Add(TakeId, MoveTemp(TakePath));
+		if (TakeInfoResult.IsSet())
+		{
+			int32 TakeId = AddTake(TakeInfoResult.GetValue());
+			FString CurrentDirectory = FPaths::GetPath(CurrentTakeFile);
 
-					This->PublishEvent<FIngestCapability_TakeAddedEvent>(TakeId);
-				}
-			}
+			FString TakePath = CurrentDirectory.Right(CurrentDirectory.Len() - StoragePath.Len());
+			FPaths::NormalizeDirectoryName(TakePath); // Removes trailing slash
+			RelativeTakePaths.Add(TakeId, MoveTemp(TakePath));
 
-			This->ExecuteUpdateTakeListCallback(Callback.Get(), Execute_GetTakeIdentifiers(This.Get()));
-		};
+			PublishEvent<FIngestCapability_TakeAddedEvent>(TakeId);
+		}
+	}
 
-	UE::Tasks::Launch(UE_SOURCE_LOCATION, MoveTemp(Task), UE::Tasks::ETaskPriority::BackgroundNormal);
+	ExecuteUpdateTakeListCallback(InCallback, Execute_GetTakeIdentifiers(this));
 }
 
 ELiveLinkDeviceConnectionStatus UTakeArchiveIngestDevice::GetConnectionStatus_Implementation() const
diff --git a/Engine/Plugins/VirtualProduction/CaptureManager/CaptureManagerDevices/Source/TakeArchiveIngestDevice/Private/TakeArchiveIngestDevice.h b/Engine/Plugins/VirtualProduction/CaptureManager/CaptureManagerDevices/Source/TakeArchiveIngestDevice/Private/TakeArchiveIngestDevice.h
--- a/Engine/Plugins/VirtualProduction/CaptureManager/CaptureManagerDevices/Source/TakeArchiveIngestDevice/Private/TakeArchiveIngestDevice.h
+++ b/Engine/Plugins/VirtualProduction/CaptureManager/CaptureManagerDevices/Source/TakeArchiveIngestDevice/Private/TakeArchiveIngestDevice.h
@@ -55,7 +55,7 @@ private:
 
 	//~ Begin ULiveLinkDeviceCapability_Ingest interface
 	virtual void RunConvertAndUploadTake(const UIngestCapability_ProcessHandle* InProcessHandle, const UIngestCapability_Options* InIngestOptions) override;
-	virtual void UpdateTakeList_Implementation(UIngestCapability_UpdateTakeListCallback* InCallback) override;
+	virtual void RunUpdateTakeList(UIngestCapability_UpdateTakeListCallback* InCallback) override;
 	//~ End ULiveLinkDeviceCapability_Ingest interface
 
 	//~ Begin ILiveLinkDeviceCapability_Connection interface
