Fixed handling of the fractional frame rate

#Synced-CL 51334653

diff --git a/Engine/Plugins/VirtualProduction/CaptureManager/CaptureManagerDevices/Source/MonoVideoIngestDevice/Private/MonoVideoIngestDevice.cpp b/Engine/Plugins/VirtualProduction/CaptureManager/CaptureManagerDevices/Source/MonoVideoIngestDevice/Private/MonoVideoIngestDevice.cpp
--- a/Engine/Plugins/VirtualProduction/CaptureManager/CaptureManagerDevices/Source/MonoVideoIngestDevice/Private/MonoVideoIngestDevice.cpp
+++ b/Engine/Plugins/VirtualProduction/CaptureManager/CaptureManagerDevices/Source/MonoVideoIngestDevice/Private/MonoVideoIngestDevice.cpp
@@ -313,7 +313,7 @@ TOptional<FTakeMetadata> UMonoVideoIngestDevice::ReadTake(const FString& InCurre
 		FCaptureExtractVideoInfo VideoFileInfo = ExtractResult.StealValue();
 		FString TimecodeStartString = VideoFileInfo.GetTimecode().ToString();
 		Video.TimecodeStart = TimecodeStartString;
-		Video.FrameRate = static_cast<uint16>(VideoFileInfo.GetFrameRate().AsDecimal());
+		Video.FrameRate = static_cast<float>(VideoFileInfo.GetFrameRate().AsDecimal());
 
 		switch (VideoFileInfo.GetVideoOrientation())
 		{
@@ -338,7 +338,7 @@ TOptional<FTakeMetadata> UMonoVideoIngestDevice::ReadTake(const FString& InCurre
 			Audio.Duration = VideoFileInfo.GetAudioDurationSeconds();
 
 			Audio.TimecodeStart = TimecodeStartString;
-			Audio.TimecodeRate = static_cast<float>(Video.FrameRate);
+			Audio.TimecodeRate = Video.FrameRate;
 
 			TakeMetadata.Audio.Add(MoveTemp(Audio));
 		}
