Persisting media sources created by LiveLinkDevices

#Synced-CL 52056680

diff --git a/Engine/Plugins/Animation/LiveLinkDevice/Source/LiveLinkDevice/Private/LiveLinkDeviceSubsystem.cpp b/Engine/Plugins/Animation/LiveLinkDevice/Source/LiveLinkDevice/Private/LiveLinkDeviceSubsystem.cpp
--- a/Engine/Plugins/Animation/LiveLinkDevice/Source/LiveLinkDevice/Private/LiveLinkDeviceSubsystem.cpp
+++ b/Engine/Plugins/Animation/LiveLinkDevice/Source/LiveLinkDevice/Private/LiveLinkDeviceSubsystem.cpp
@@ -88,6 +88,7 @@ void ULiveLinkDeviceSubsystem::OnExtraDataSessionSaving(ULiveLinkHubSessionExtra
 			.DeviceGuid = DevicePair.Key,
 			.DeviceClass = DevicePair.Value->GetClass(),
 			.DeviceSettings = DevicePair.Value->GetDeviceSettings(),
+			.DeviceCapabilityData = DevicePair.Value->OnSaveDeviceCapabilityData()
 		});
 	}
 }
@@ -137,7 +138,12 @@ void ULiveLinkDeviceSubsystem::OnExtraDataSessionLoaded(const ULiveLinkHubSessio
 			DevicePreset.DeviceGuid
 		);
 
-		if (!CreateResult.HasValue())
+		if (CreateResult.HasValue())
+		{
+			FLiveLinkDeviceCreateResult Result = CreateResult.StealValue();
+			Result.Device->OnLoadDeviceCapabilityData(DevicePreset.DeviceCapabilityData);
+		}
+		else
 		{
 			UE_LOGFMT(LogLiveLinkDevice, Error, "Failed to restore device");
 		}
diff --git a/Engine/Plugins/Animation/LiveLinkDevice/Source/LiveLinkDevice/Private/LiveLinkHubSessionExtraData_Device.h b/Engine/Plugins/Animation/LiveLinkDevice/Source/LiveLinkDevice/Private/LiveLinkHubSessionExtraData_Device.h
--- a/Engine/Plugins/Animation/LiveLinkDevice/Source/LiveLinkDevice/Private/LiveLinkHubSessionExtraData_Device.h
+++ b/Engine/Plugins/Animation/LiveLinkDevice/Source/LiveLinkDevice/Private/LiveLinkHubSessionExtraData_Device.h
@@ -9,7 +9,7 @@
 
 class ULiveLinkDevice;
 class ULiveLinkDeviceSettings;
-
+class ULiveLinkDeviceCapability;
 
 USTRUCT()
 struct FLiveLinkDevicePreset
@@ -24,6 +24,9 @@ struct FLiveLinkDevicePreset
 
 	UPROPERTY(Instanced)
 	TObjectPtr<ULiveLinkDeviceSettings> DeviceSettings;
+
+	UPROPERTY(Instanced)
+	TMap<TSubclassOf<ULiveLinkDeviceCapability>, TObjectPtr<UObject>> DeviceCapabilityData;
 };
 
 
diff --git a/Engine/Plugins/Animation/LiveLinkDevice/Source/LiveLinkDevice/Public/LiveLinkDevice.h b/Engine/Plugins/Animation/LiveLinkDevice/Source/LiveLinkDevice/Public/LiveLinkDevice.h
--- a/Engine/Plugins/Animation/LiveLinkDevice/Source/LiveLinkDevice/Public/LiveLinkDevice.h
+++ b/Engine/Plugins/Animation/LiveLinkDevice/Source/LiveLinkDevice/Public/LiveLinkDevice.h
@@ -73,6 +73,12 @@ public:
 	/** Called when a setting value has changed via the UI. */
 	virtual void OnSettingChanged(const FPropertyChangedEvent& InPropertyChangedEvent) { }
 
+	/** Called when the device is being saved to the settings */
+	virtual TMap<TSubclassOf<class ULiveLinkDeviceCapability>, TObjectPtr<UObject>> OnSaveDeviceCapabilityData() const { return {}; }
+
+	/** Called when the device is being loaded from the settings */
+	virtual void OnLoadDeviceCapabilityData(TMap<TSubclassOf<class ULiveLinkDeviceCapability>, TObjectPtr<UObject>> InCapabilityDataMap) { }
+
 	/** This device's settings UObject. */
 	template<typename T = ULiveLinkDeviceSettings>
 	const T* GetDeviceSettings() const { return CastChecked<T>(Settings); }
diff --git a/Engine/Plugins/VirtualProduction/CaptureManager/CaptureManagerDevices/Source/CPSLiveLinkDevice/Private/LiveLinkFaceDevice.cpp b/Engine/Plugins/VirtualProduction/CaptureManager/CaptureManagerDevices/Source/CPSLiveLinkDevice/Private/LiveLinkFaceDevice.cpp
--- a/Engine/Plugins/VirtualProduction/CaptureManager/CaptureManagerDevices/Source/CPSLiveLinkDevice/Private/LiveLinkFaceDevice.cpp
+++ b/Engine/Plugins/VirtualProduction/CaptureManager/CaptureManagerDevices/Source/CPSLiveLinkDevice/Private/LiveLinkFaceDevice.cpp
@@ -66,6 +66,32 @@ void ULiveLinkFaceDevice::OnDeviceRemoved()
 	Super::OnDeviceRemoved();
 }
 
+TMap<TSubclassOf<ULiveLinkDeviceCapability>, TObjectPtr<UObject>> ULiveLinkFaceDevice::OnSaveDeviceCapabilityData() const
+{
+	TMap<TSubclassOf<ULiveLinkDeviceCapability>, TObjectPtr<UObject>> CapabilityData;
+
+	if (TObjectPtr<URtspMediaSource> MediaSource = Cast<URtspMediaSource>(GetMediaSource_Implementation()))
+	{
+		CapabilityData.Add(ULiveLinkDeviceCapability_Streaming::StaticClass(), MediaSource);
+	}
+	
+	return CapabilityData;
+}
+
+void ULiveLinkFaceDevice::OnLoadDeviceCapabilityData(TMap<TSubclassOf<ULiveLinkDeviceCapability>, TObjectPtr<UObject>> InCapabilityDataMap)
+{
+	for (const TPair<TSubclassOf<ULiveLinkDeviceCapability>, TObjectPtr<UObject>>& CapabilityData : InCapabilityDataMap)
+	{
+		if (CapabilityData.Key == ULiveLinkDeviceCapability_Streaming::StaticClass())
+		{
+			if (TObjectPtr<URtspMediaSource> MediaSource = Cast<URtspMediaSource>(CapabilityData.Value))
+			{
+				RegisterMediaSource(MediaSource);
+			}
+		}
+	}
+}
+
 FString ULiveLinkFaceDevice::GetFullTakePath(UE::CaptureManager::FTakeId InTakeId) const
 {
 	FScopeLock Lock(&DownloadedTakesMutex);
@@ -327,11 +353,15 @@ void ULiveLinkFaceDevice::HandleConnectionChanged(TSharedPtr<const UE::CaptureMa
 	}
 	else if (Event->ConnectionState == FConnectionStateChangedEvent::EState::Connected)
 	{
-		TObjectPtr<URtspMediaSource> MediaSource = NewObject<URtspMediaSource>();
-		MediaSource->Host = GetSettings()->IpAddress.IpAddressString;
-		MediaSource->DecoderBufferSize = 1; // Using 1 for minimum latency
+		// Don't register if it is already registered
+		if (!GetMediaSource_Implementation())
+		{
+			TObjectPtr<URtspMediaSource> RtspMediaSource = NewObject<URtspMediaSource>();
+			RtspMediaSource->Host = GetSettings()->IpAddress.IpAddressString;
+			RtspMediaSource->DecoderBufferSize = 1; // Using 1 for minimum latency
 
-		RegisterMediaSource(MediaSource);
+			RegisterMediaSource(RtspMediaSource);
+		}
 
 		bIsConnecting = false;
 	}
diff --git a/Engine/Plugins/VirtualProduction/CaptureManager/CaptureManagerDevices/Source/CPSLiveLinkDevice/Public/LiveLinkFaceDevice.h b/Engine/Plugins/VirtualProduction/CaptureManager/CaptureManagerDevices/Source/CPSLiveLinkDevice/Public/LiveLinkFaceDevice.h
--- a/Engine/Plugins/VirtualProduction/CaptureManager/CaptureManagerDevices/Source/CPSLiveLinkDevice/Public/LiveLinkFaceDevice.h
+++ b/Engine/Plugins/VirtualProduction/CaptureManager/CaptureManagerDevices/Source/CPSLiveLinkDevice/Public/LiveLinkFaceDevice.h
@@ -61,6 +61,8 @@ public:
 
 	virtual void OnDeviceAdded() override;
 	virtual void OnDeviceRemoved() override;
+	virtual TMap<TSubclassOf<ULiveLinkDeviceCapability>, TObjectPtr<UObject>> OnSaveDeviceCapabilityData() const override;
+	virtual void OnLoadDeviceCapabilityData(TMap<TSubclassOf<ULiveLinkDeviceCapability>, TObjectPtr<UObject>> InCapabilityDataMap) override;
 	//~ End ULiveLinkDevice interface
 
 private:
